Moving WebBrowser and CEF3Utils to Runtime so that they can be used in games.

Also removed guards around browser code in SWebBrowser and made sure that path for sub process exe is correct if game exe is not in the same folder.

diff --git a/Engine/Source/Developer/CEF3Utils/CEF3Utils.Build.cs b/Engine/Source/Developer/CEF3Utils/CEF3Utils.Build.cs
deleted file mode 100644
diff --git a/Engine/Source/Developer/CEF3Utils/Private/CEF3Utils.cpp b/Engine/Source/Developer/CEF3Utils/Private/CEF3Utils.cpp
deleted file mode 100644
diff --git a/Engine/Source/Developer/CEF3Utils/Private/CEF3UtilsPrivatePCH.h b/Engine/Source/Developer/CEF3Utils/Private/CEF3UtilsPrivatePCH.h
deleted file mode 100644
diff --git a/Engine/Source/Developer/CEF3Utils/Public/CEF3Utils.h b/Engine/Source/Developer/CEF3Utils/Public/CEF3Utils.h
deleted file mode 100644
diff --git a/Engine/Source/Developer/WebBrowser/Private/SWebBrowser.cpp b/Engine/Source/Developer/WebBrowser/Private/SWebBrowser.cpp
deleted file mode 100644
diff --git a/Engine/Source/Developer/WebBrowser/Private/WebBrowserApp.cpp b/Engine/Source/Developer/WebBrowser/Private/WebBrowserApp.cpp
deleted file mode 100644
diff --git a/Engine/Source/Developer/WebBrowser/Private/WebBrowserApp.h b/Engine/Source/Developer/WebBrowser/Private/WebBrowserApp.h
deleted file mode 100644
diff --git a/Engine/Source/Developer/WebBrowser/Private/WebBrowserHandler.cpp b/Engine/Source/Developer/WebBrowser/Private/WebBrowserHandler.cpp
deleted file mode 100644
diff --git a/Engine/Source/Developer/WebBrowser/Private/WebBrowserHandler.h b/Engine/Source/Developer/WebBrowser/Private/WebBrowserHandler.h
deleted file mode 100644
diff --git a/Engine/Source/Developer/WebBrowser/Private/WebBrowserModule.cpp b/Engine/Source/Developer/WebBrowser/Private/WebBrowserModule.cpp
deleted file mode 100644
diff --git a/Engine/Source/Developer/WebBrowser/Private/WebBrowserPrivatePCH.h b/Engine/Source/Developer/WebBrowser/Private/WebBrowserPrivatePCH.h
deleted file mode 100644
diff --git a/Engine/Source/Developer/WebBrowser/Private/WebBrowserSingleton.cpp b/Engine/Source/Developer/WebBrowser/Private/WebBrowserSingleton.cpp
deleted file mode 100644
diff --git a/Engine/Source/Developer/WebBrowser/Private/WebBrowserSingleton.h b/Engine/Source/Developer/WebBrowser/Private/WebBrowserSingleton.h
deleted file mode 100644
diff --git a/Engine/Source/Developer/WebBrowser/Private/WebBrowserViewport.cpp b/Engine/Source/Developer/WebBrowser/Private/WebBrowserViewport.cpp
deleted file mode 100644
diff --git a/Engine/Source/Developer/WebBrowser/Private/WebBrowserWindow.cpp b/Engine/Source/Developer/WebBrowser/Private/WebBrowserWindow.cpp
deleted file mode 100644
diff --git a/Engine/Source/Developer/WebBrowser/Private/WebBrowserWindow.h b/Engine/Source/Developer/WebBrowser/Private/WebBrowserWindow.h
deleted file mode 100644
diff --git a/Engine/Source/Developer/WebBrowser/Public/IWebBrowserSingleton.h b/Engine/Source/Developer/WebBrowser/Public/IWebBrowserSingleton.h
deleted file mode 100644
diff --git a/Engine/Source/Developer/WebBrowser/Public/IWebBrowserWindow.h b/Engine/Source/Developer/WebBrowser/Public/IWebBrowserWindow.h
deleted file mode 100644
diff --git a/Engine/Source/Developer/WebBrowser/Public/SWebBrowser.h b/Engine/Source/Developer/WebBrowser/Public/SWebBrowser.h
deleted file mode 100644
diff --git a/Engine/Source/Developer/WebBrowser/Public/WebBrowserModule.h b/Engine/Source/Developer/WebBrowser/Public/WebBrowserModule.h
deleted file mode 100644
diff --git a/Engine/Source/Developer/WebBrowser/Public/WebBrowserViewport.h b/Engine/Source/Developer/WebBrowser/Public/WebBrowserViewport.h
deleted file mode 100644
diff --git a/Engine/Source/Developer/WebBrowser/WebBrowser.Build.cs b/Engine/Source/Developer/WebBrowser/WebBrowser.Build.cs
deleted file mode 100644
diff --git a/Engine/Source/Runtime/CEF3Utils/CEF3Utils.Build.cs b/Engine/Source/Runtime/CEF3Utils/CEF3Utils.Build.cs
new file mode 100644
--- /dev/null
+++ b/Engine/Source/Runtime/CEF3Utils/CEF3Utils.Build.cs
@@ -0,0 +1,27 @@
+// Copyright 1998-2015 Epic Games, Inc. All Rights Reserved.
+
+using UnrealBuildTool;
+
+public class CEF3Utils : ModuleRules
+{
+	public CEF3Utils(TargetInfo Target)
+	{
+		PublicIncludePaths.Add("Developer/CEF3Utils/Public");
+		PrivateIncludePaths.Add("Developer/CEF3Utils/Private");
+
+		PrivateDependencyModuleNames.AddRange(
+			new string[]
+			{
+				"Core",
+			}
+		);
+
+		if (Target.Platform == UnrealTargetPlatform.Win64
+        ||  Target.Platform == UnrealTargetPlatform.Mac)
+		{
+			AddThirdPartyPrivateStaticDependencies(Target,
+				"CEF3"
+				);
+		}
+	}
+}
diff --git a/Engine/Source/Runtime/CEF3Utils/Private/CEF3Utils.cpp b/Engine/Source/Runtime/CEF3Utils/Private/CEF3Utils.cpp
new file mode 100644
--- /dev/null
+++ b/Engine/Source/Runtime/CEF3Utils/Private/CEF3Utils.cpp
@@ -0,0 +1,67 @@
+// Copyright 1998-2015 Epic Games, Inc. All Rights Reserved.
+
+#include "CEF3UtilsPrivatePCH.h"
+
+DEFINE_LOG_CATEGORY(LogCEF3Utils);
+
+IMPLEMENT_MODULE(FDefaultModuleImpl, CEF3Utils);
+
+#if WITH_CEF3
+namespace CEF3Utils
+{
+    void* CEF3DLLHandle = nullptr;
+#if PLATFORM_WINDOWS
+	void* D3DHandle = nullptr;
+	void* MPEGHandle = nullptr;
+	void* ICUDTHandle = nullptr;
+	void* GLESHandle = nullptr;
+	void* EGLHandle = nullptr;
+#endif
+
+	void LoadCEF3Modules()
+	{
+#if PLATFORM_WINDOWS
+		FString DllPath(FPaths::Combine(*FPaths::EngineDir(), TEXT("Binaries/ThirdParty/CEF3/Win64")));
+
+		CEF3DLLHandle = FPlatformProcess::GetDllHandle(*FPaths::Combine(*DllPath, TEXT("libcef.dll")));
+		if (CEF3DLLHandle == NULL)
+		{
+			int32 ErrorNum = FPlatformMisc::GetLastError();
+			TCHAR ErrorMsg[1024];
+			FPlatformMisc::GetSystemErrorMessage(ErrorMsg, 1024, ErrorNum);
+			UE_LOG(LogCEF3Utils, Error, TEXT("Failed to get CEF3 DLL handle: %s (%d)"), ErrorMsg, ErrorNum);
+		}
+
+	#if WINVER >= 0x600 // Different dll used pre-Vista
+		D3DHandle = FPlatformProcess::GetDllHandle(*FPaths::Combine(*DllPath, TEXT("d3dcompiler_46.dll")));
+	#else
+		D3DHandle = FPlatformProcess::GetDllHandle(*FPaths::Combine(*DllPath, TEXT("d3dcompiler_43.dll")));
+	#endif
+		MPEGHandle = FPlatformProcess::GetDllHandle(*FPaths::Combine(*DllPath, TEXT("ffmpegsumo.dll")));
+		ICUDTHandle = FPlatformProcess::GetDllHandle(*FPaths::Combine(*DllPath, TEXT("icudt.dll")));
+		GLESHandle = FPlatformProcess::GetDllHandle(*FPaths::Combine(*DllPath, TEXT("libGLESv2.dll")));
+		EGLHandle = FPlatformProcess::GetDllHandle(*FPaths::Combine(*DllPath, TEXT("libEGL.dll")));
+#elif PLATFORM_MAC
+        CEF3DLLHandle = FPlatformProcess::GetDllHandle(TEXT("libplugin_carbon_interpose.dylib"));
+#endif
+	}
+
+	void UnloadCEF3Modules()
+	{
+		FPlatformProcess::FreeDllHandle(CEF3DLLHandle);
+		CEF3DLLHandle = nullptr;
+#if PLATFORM_WINDOWS
+		FPlatformProcess::FreeDllHandle(D3DHandle);
+		D3DHandle = nullptr;
+		FPlatformProcess::FreeDllHandle(MPEGHandle);
+		MPEGHandle = nullptr;
+		FPlatformProcess::FreeDllHandle(ICUDTHandle);
+		ICUDTHandle = nullptr;
+		FPlatformProcess::FreeDllHandle(GLESHandle);
+		GLESHandle = nullptr;
+		FPlatformProcess::FreeDllHandle(EGLHandle);
+		EGLHandle = nullptr;
+#endif
+	}
+};
+#endif //WITH_CEF3
diff --git a/Engine/Source/Runtime/CEF3Utils/Private/CEF3UtilsPrivatePCH.h b/Engine/Source/Runtime/CEF3Utils/Private/CEF3UtilsPrivatePCH.h
new file mode 100644
--- /dev/null
+++ b/Engine/Source/Runtime/CEF3Utils/Private/CEF3UtilsPrivatePCH.h
@@ -0,0 +1,9 @@
+// Copyright 1998-2015 Epic Games, Inc. All Rights Reserved.
+
+#pragma once
+
+#include "Core.h"
+#include "ModuleManager.h"
+#include "CEF3Utils.h"
+
+DECLARE_LOG_CATEGORY_EXTERN(LogCEF3Utils, Log, All);
\ No newline at end of file
diff --git a/Engine/Source/Runtime/CEF3Utils/Public/CEF3Utils.h b/Engine/Source/Runtime/CEF3Utils/Public/CEF3Utils.h
new file mode 100644
--- /dev/null
+++ b/Engine/Source/Runtime/CEF3Utils/Public/CEF3Utils.h
@@ -0,0 +1,18 @@
+// Copyright 1998-2015 Epic Games, Inc. All Rights Reserved.
+
+#pragma once
+
+#if WITH_CEF3
+namespace CEF3Utils
+{
+	/**
+	 * Load the required modules for CEF3
+	 */
+	CEF3UTILS_API void LoadCEF3Modules();
+
+	/**
+	 * Unload the required modules for CEF3
+	 */
+	CEF3UTILS_API void UnloadCEF3Modules();
+};
+#endif //WITH_CEF3
diff --git a/Engine/Source/Runtime/WebBrowser/Private/SWebBrowser.cpp b/Engine/Source/Runtime/WebBrowser/Private/SWebBrowser.cpp
new file mode 100644
--- /dev/null
+++ b/Engine/Source/Runtime/WebBrowser/Private/SWebBrowser.cpp
@@ -0,0 +1,202 @@
+// Copyright 1998-2015 Epic Games, Inc. All Rights Reserved.
+
+#include "WebBrowserPrivatePCH.h"
+#include "SWebBrowser.h"
+#include "SThrobber.h"
+#include "WebBrowserModule.h"
+#include "IWebBrowserSingleton.h"
+#include "IWebBrowserWindow.h"
+#include "WebBrowserViewport.h"
+
+#define LOCTEXT_NAMESPACE "WebBrowser"
+
+SWebBrowser::SWebBrowser()
+{
+}
+
+void SWebBrowser::Construct(const FArguments& InArgs)
+{
+	void* OSWindowHandle = nullptr;
+	if (InArgs._ParentWindow.IsValid())
+	{
+		TSharedPtr<FGenericWindow> NativeWindow = InArgs._ParentWindow->GetNativeWindow();
+		OSWindowHandle = NativeWindow->GetOSWindowHandle();
+	}
+
+	BrowserWindow = IWebBrowserModule::Get().GetSingleton()->CreateBrowserWindow(OSWindowHandle,
+	                                                                             InArgs._InitialURL,
+	                                                                             InArgs._ViewportSize.Get().X,
+	                                                                             InArgs._ViewportSize.Get().Y,
+	                                                                             InArgs._SupportsTransparency);
+
+	TSharedPtr<SViewport> ViewportWidget;
+
+	ChildSlot
+	[
+		SNew(SVerticalBox)
+		+SVerticalBox::Slot()
+		.AutoHeight()
+		.Padding(0, 5)
+		[
+			SNew(SHorizontalBox)
+			.Visibility(InArgs._ShowControls ? EVisibility::Visible : EVisibility::Collapsed)
+			+SHorizontalBox::Slot()
+			.AutoWidth()
+			[
+				SNew(SButton)
+				.Text(LOCTEXT("Back","Back"))
+				.IsEnabled(this, &SWebBrowser::CanGoBack)
+				.OnClicked(this, &SWebBrowser::OnBackClicked)
+			]
+			+SHorizontalBox::Slot()
+			.AutoWidth()
+			[
+				SNew(SButton)
+				.Text(LOCTEXT("Forward", "Forward"))
+				.IsEnabled(this, &SWebBrowser::CanGoForward)
+				.OnClicked(this, &SWebBrowser::OnForwardClicked)
+			]
+			+SHorizontalBox::Slot()
+			.AutoWidth()
+			[
+				SNew(SButton)
+				.Text(this, &SWebBrowser::GetReloadButtonText)
+				.OnClicked(this, &SWebBrowser::OnReloadClicked)
+			]
+			+SHorizontalBox::Slot()
+			.FillWidth(1.0f)
+			.VAlign(VAlign_Center)
+			.HAlign(HAlign_Right)
+			.Padding(5)
+			[
+				SNew(STextBlock)
+				.Text(this, &SWebBrowser::GetTitleText)
+				.Justification(ETextJustify::Right)
+			]
+		]
+		+SVerticalBox::Slot()
+		[
+			SNew(SOverlay)
+			+SOverlay::Slot()
+			[
+				SAssignNew(ViewportWidget, SViewport)
+				.ViewportSize(InArgs._ViewportSize)
+				.EnableGammaCorrection(false)
+				.EnableBlending(InArgs._SupportsTransparency)
+				.IgnoreTextureAlpha(!InArgs._SupportsTransparency)
+				.Visibility(this, &SWebBrowser::GetViewportVisibility)
+			]
+			+SOverlay::Slot()
+			.HAlign(HAlign_Center)
+			.VAlign(VAlign_Center)
+			[
+				SNew(SCircularThrobber)
+				.Radius(10.0f)
+				.ToolTipText(LOCTEXT("LoadingThrobberToolTip", "Loading page..."))
+				.Visibility(this, &SWebBrowser::GetLoadingThrobberVisibility)
+			]
+		]
+	];
+
+	if (BrowserWindow.IsValid())
+	{
+		BrowserViewport = MakeShareable(new FWebBrowserViewport(BrowserWindow, ViewportWidget));
+		ViewportWidget->SetViewportInterface(BrowserViewport.ToSharedRef());
+	}
+}
+
+FText SWebBrowser::GetTitleText() const
+{
+	if (BrowserWindow.IsValid())
+	{
+		return FText::FromString(BrowserWindow->GetTitle());
+	}
+	return LOCTEXT("InvalidWindow", "Browser Window is not valid/supported");
+}
+
+bool SWebBrowser::CanGoBack() const
+{
+	if (BrowserWindow.IsValid())
+	{
+		return BrowserWindow->CanGoBack();
+	}
+	return false;
+}
+
+FReply SWebBrowser::OnBackClicked()
+{
+	if (BrowserWindow.IsValid())
+	{
+		BrowserWindow->GoBack();
+	}
+	return FReply::Handled();
+}
+
+bool SWebBrowser::CanGoForward() const
+{
+	if (BrowserWindow.IsValid())
+	{
+		return BrowserWindow->CanGoForward();
+	}
+	return false;
+}
+
+FReply SWebBrowser::OnForwardClicked()
+{
+	if (BrowserWindow.IsValid())
+	{
+		BrowserWindow->GoForward();
+	}
+	return FReply::Handled();
+}
+
+FText SWebBrowser::GetReloadButtonText() const
+{
+	static FText Reload = LOCTEXT("Reload", "Reload");
+	static FText Stop = LOCTEXT("Stop", "Stop");
+
+	if (BrowserWindow.IsValid())
+	{
+		if (BrowserWindow->IsLoading())
+		{
+			return Stop;
+		}
+	}
+	return Reload;
+}
+
+FReply SWebBrowser::OnReloadClicked()
+{
+	if (BrowserWindow.IsValid())
+	{
+		if (BrowserWindow->IsLoading())
+		{
+			BrowserWindow->StopLoad();
+		}
+		else
+		{
+			BrowserWindow->Reload();
+		}
+	}
+	return FReply::Handled();
+}
+
+EVisibility SWebBrowser::GetViewportVisibility() const
+{
+	if (BrowserWindow.IsValid() && BrowserWindow->HasBeenPainted())
+	{
+		return EVisibility::Visible;
+	}
+	return EVisibility::Hidden;
+}
+
+EVisibility SWebBrowser::GetLoadingThrobberVisibility() const
+{
+	if (BrowserWindow.IsValid() && !BrowserWindow->HasBeenPainted())
+	{
+		return EVisibility::Visible;
+	}
+	return EVisibility::Hidden;
+}
+
+#undef LOCTEXT_NAMESPACE
diff --git a/Engine/Source/Runtime/WebBrowser/Private/WebBrowserApp.cpp b/Engine/Source/Runtime/WebBrowser/Private/WebBrowserApp.cpp
new file mode 100644
--- /dev/null
+++ b/Engine/Source/Runtime/WebBrowser/Private/WebBrowserApp.cpp
@@ -0,0 +1,16 @@
+// Copyright 1998-2015 Epic Games, Inc. All Rights Reserved.
+
+#include "WebBrowserPrivatePCH.h"
+#include "WebBrowserApp.h"
+
+#if WITH_CEF3
+FWebBrowserApp::FWebBrowserApp()
+{
+}
+
+void FWebBrowserApp::OnContextInitialized()
+{
+	// Called on the browser process UI thread immediately after the CEF context
+	// has been initialized.
+}
+#endif
diff --git a/Engine/Source/Runtime/WebBrowser/Private/WebBrowserApp.h b/Engine/Source/Runtime/WebBrowser/Private/WebBrowserApp.h
new file mode 100644
--- /dev/null
+++ b/Engine/Source/Runtime/WebBrowser/Private/WebBrowserApp.h
@@ -0,0 +1,39 @@
+// Copyright 1998-2015 Epic Games, Inc. All Rights Reserved.
+
+#pragma once
+
+#if WITH_CEF3
+
+#if PLATFORM_WINDOWS
+#include "AllowWindowsPlatformTypes.h"
+#endif
+#include "include/cef_app.h"
+#if PLATFORM_WINDOWS
+#include "HideWindowsPlatformTypes.h"
+#endif
+
+/**
+ * Implements CEF App and other Process level interfaces 
+ */
+class FWebBrowserApp : public CefApp,
+	public CefBrowserProcessHandler,
+	public CefRenderProcessHandler
+{
+public:
+	
+	/**
+	 * Default Constructor
+	 */
+	FWebBrowserApp();
+private:
+	// CefApp methods.
+	virtual CefRefPtr<CefBrowserProcessHandler> GetBrowserProcessHandler() override { return this; }
+	virtual CefRefPtr<CefRenderProcessHandler> GetRenderProcessHandler() override { return this; }
+
+	// CefBrowserProcessHandler methods:
+	virtual void OnContextInitialized() override;
+
+	// Include the default reference counting implementation.
+	IMPLEMENT_REFCOUNTING(FWebBrowserApp);
+};
+#endif
diff --git a/Engine/Source/Runtime/WebBrowser/Private/WebBrowserHandler.cpp b/Engine/Source/Runtime/WebBrowser/Private/WebBrowserHandler.cpp
new file mode 100644
--- /dev/null
+++ b/Engine/Source/Runtime/WebBrowser/Private/WebBrowserHandler.cpp
@@ -0,0 +1,115 @@
+// Copyright 1998-2015 Epic Games, Inc. All Rights Reserved.
+
+#include "WebBrowserPrivatePCH.h"
+#include "WebBrowserHandler.h"
+#include "WebBrowserWindow.h"
+#include "WebBrowserSingleton.h"
+
+#define LOCTEXT_NAMESPACE "WebBrowserHandler"
+
+#if WITH_CEF3
+FWebBrowserHandler::FWebBrowserHandler()
+{
+}
+
+void FWebBrowserHandler::OnTitleChange(CefRefPtr<CefBrowser> Browser, const CefString& Title)
+{
+	if (BrowserWindow.IsValid())
+	{
+		BrowserWindow.Pin()->SetTitle(Title);
+	}
+}
+
+void FWebBrowserHandler::OnAfterCreated(CefRefPtr<CefBrowser> Browser)
+{
+	if (BrowserWindow.IsValid())
+	{
+		BrowserWindow.Pin()->BindCefBrowser(Browser);
+	}
+}
+
+void FWebBrowserHandler::OnBeforeClose(CefRefPtr<CefBrowser> Browser)
+{
+	if (BrowserWindow.IsValid())
+	{
+		BrowserWindow.Pin()->BindCefBrowser(nullptr);
+	}
+}
+
+void FWebBrowserHandler::OnLoadError(CefRefPtr<CefBrowser> Browser,
+	CefRefPtr<CefFrame> Frame,
+	CefLoadHandler::ErrorCode InErrorCode,
+	const CefString& ErrorText,
+	const CefString& FailedUrl)
+{
+	// Don't display an error for downloaded files.
+	if (InErrorCode == ERR_ABORTED)
+		return;
+
+	// Display a load error message.
+	FFormatNamedArguments Args;
+	Args.Add(TEXT("FailedUrl"), FText::FromString(FailedUrl.ToWString().c_str()));
+	Args.Add(TEXT("ErrorText"), FText::FromString(ErrorText.ToWString().c_str()));
+	Args.Add(TEXT("ErrorCode"), FText::AsNumber(InErrorCode));
+	FText ErrorMsg = FText::Format(LOCTEXT("WebBrowserLoadError", "Failed to load URL {FailedUrl} with error {ErrorText} ({ErrorCode})."), Args);
+	FString ErrorHTML = TEXT("<html><body bgcolor=\"white\"><h2>")
+						+ ErrorMsg.ToString()
+						+ TEXT("</h2></body></html>");
+	Frame->LoadString(*ErrorHTML, FailedUrl);
+}
+
+bool FWebBrowserHandler::GetViewRect(CefRefPtr<CefBrowser> Browser, CefRect& Rect)
+{
+	if (BrowserWindow.IsValid())
+	{
+		return BrowserWindow.Pin()->GetViewRect(Rect);
+	}
+	return false;
+}
+
+void FWebBrowserHandler::OnPaint(CefRefPtr<CefBrowser> Browser,
+	PaintElementType Type,
+	const RectList& DirtyRects,
+	const void* Buffer,
+	int Width, int Height)
+{
+	if (BrowserWindow.IsValid())
+	{
+		BrowserWindow.Pin()->OnPaint(Type, DirtyRects, Buffer, Width, Height);
+	}
+}
+
+void FWebBrowserHandler::OnCursorChange(CefRefPtr<CefBrowser> Browser, CefCursorHandle Cursor)
+{
+	if (BrowserWindow.IsValid())
+	{
+		BrowserWindow.Pin()->OnCursorChange(Cursor);
+	}
+}
+
+bool FWebBrowserHandler::OnBeforeResourceLoad(CefRefPtr<CefBrowser> Browser, CefRefPtr<CefFrame> Frame, CefRefPtr<CefRequest> Request)
+{
+	const FString LanguageHeaderText(TEXT("Accept-Language"));
+	const FString LocaleCode = FWebBrowserSingleton::GetCurrentLocaleCode();
+	CefRequest::HeaderMap HeaderMap;
+	Request->GetHeaderMap(HeaderMap);
+	auto LanguageHeader = HeaderMap.find(*LanguageHeaderText);
+	if (LanguageHeader != HeaderMap.end())
+	{
+		(*LanguageHeader).second = *LocaleCode;
+	}
+	else
+	{
+		HeaderMap.insert(std::pair<CefString, CefString>(*LanguageHeaderText, *LocaleCode));
+	}
+	Request->SetHeaderMap(HeaderMap);
+	return false;
+}
+
+void FWebBrowserHandler::SetBrowserWindow(TSharedPtr<FWebBrowserWindow> InBrowserWindow)
+{
+	BrowserWindow = InBrowserWindow;
+}
+#endif
+
+#undef LOCTEXT_NAMESPACE
diff --git a/Engine/Source/Runtime/WebBrowser/Private/WebBrowserHandler.h b/Engine/Source/Runtime/WebBrowser/Private/WebBrowserHandler.h
new file mode 100644
--- /dev/null
+++ b/Engine/Source/Runtime/WebBrowser/Private/WebBrowserHandler.h
@@ -0,0 +1,98 @@
+// Copyright 1998-2015 Epic Games, Inc. All Rights Reserved.
+
+#pragma once
+
+#if WITH_CEF3
+
+#if PLATFORM_WINDOWS
+#include "AllowWindowsPlatformTypes.h"
+#endif
+#include "include/cef_client.h"
+#if PLATFORM_WINDOWS
+#include "HideWindowsPlatformTypes.h"
+#endif
+
+// Forward Declarations
+class FWebBrowserWindow;
+
+/**
+ * Implements CEF Client and other Browser level interfaces 
+ */
+class FWebBrowserHandler : public CefClient,
+	public CefDisplayHandler,
+	public CefLifeSpanHandler,
+	public CefLoadHandler,
+	public CefRenderHandler,
+	public CefRequestHandler
+{
+public:
+	/**
+	 * Default Constructor
+	 */
+	FWebBrowserHandler();
+
+	// CefClient Interface
+	virtual CefRefPtr<CefDisplayHandler> GetDisplayHandler() override
+	{
+		return this;
+	}
+	virtual CefRefPtr<CefLifeSpanHandler> GetLifeSpanHandler() override
+	{
+		return this;
+	}
+	virtual CefRefPtr<CefLoadHandler> GetLoadHandler() override
+	{
+		return this;
+	}
+	virtual CefRefPtr<CefRenderHandler> GetRenderHandler() override
+	{
+		return this;
+	}
+	virtual CefRefPtr<CefRequestHandler> GetRequestHandler() override
+	{
+		return this;
+	}
+
+	// CefDisplayHandler Interface
+	virtual void OnTitleChange(CefRefPtr<CefBrowser> Browser, const CefString& Title) override;
+
+	// CefLifeSpanHandler Interface
+	virtual void OnAfterCreated(CefRefPtr<CefBrowser> Browser) override;
+	virtual void OnBeforeClose(CefRefPtr<CefBrowser> Browser) override;
+
+	// CefLoadHandler Interface
+	virtual void OnLoadError(CefRefPtr<CefBrowser> Browser,
+		CefRefPtr<CefFrame> Frame,
+		CefLoadHandler::ErrorCode InErrorCode,
+		const CefString& ErrorText,
+		const CefString& FailedUrl) override;
+
+	// CefRenderHandler Interface
+	virtual bool GetViewRect(CefRefPtr<CefBrowser> Browser, CefRect& Rect) override;
+	virtual void OnPaint(CefRefPtr<CefBrowser> Browser,
+		PaintElementType Type,
+		const RectList& DirtyRects,
+		const void* Buffer,
+		int Width, int Height) override;
+	virtual void OnCursorChange(CefRefPtr<CefBrowser> Browser, CefCursorHandle Cursor) override;
+
+	// CefRequestHandler Interface
+	virtual bool OnBeforeResourceLoad(CefRefPtr<CefBrowser> Browser,
+		CefRefPtr<CefFrame> Frame,
+		CefRefPtr<CefRequest> Request) override;
+
+	/**
+	 * Pass in a pointer to our Browser Window so that events can be passed on
+	 *
+	 * @param InBrowserWindow The browser window this will be handling
+	 */
+	void SetBrowserWindow(TSharedPtr<FWebBrowserWindow> InBrowserWindow);
+
+private:
+	/** Weak Pointer to our Web Browser window so that events can be passed on while it's valid*/
+	TWeakPtr<FWebBrowserWindow> BrowserWindow;
+
+	// Include the default reference counting implementation.
+	IMPLEMENT_REFCOUNTING(FWebBrowserHandler);
+};
+#endif
diff --git a/Engine/Source/Runtime/WebBrowser/Private/WebBrowserModule.cpp b/Engine/Source/Runtime/WebBrowser/Private/WebBrowserModule.cpp
new file mode 100644
--- /dev/null
+++ b/Engine/Source/Runtime/WebBrowser/Private/WebBrowserModule.cpp
@@ -0,0 +1,47 @@
+// Copyright 1998-2015 Epic Games, Inc. All Rights Reserved.
+
+#include "WebBrowserPrivatePCH.h"
+#include "WebBrowserSingleton.h"
+#include "ModuleManager.h"
+#include "CEF3Utils.h"
+
+DEFINE_LOG_CATEGORY(LogWebBrowser);
+
+static FWebBrowserSingleton* WebBrowserSingleton = nullptr;
+
+class FWebBrowserModule : public IWebBrowserModule
+{
+private:
+	// IModuleInterface Interface
+	virtual void StartupModule() override;
+	virtual void ShutdownModule() override;
+
+public:
+	virtual IWebBrowserSingleton* GetSingleton();
+};
+
+IMPLEMENT_MODULE( FWebBrowserModule, WebBrowser );
+
+void FWebBrowserModule::StartupModule()
+{
+#if WITH_CEF3
+	CEF3Utils::LoadCEF3Modules();
+#endif
+
+	WebBrowserSingleton = new FWebBrowserSingleton();
+}
+
+void FWebBrowserModule::ShutdownModule()
+{
+	delete WebBrowserSingleton;
+	WebBrowserSingleton = NULL;
+
+#if WITH_CEF3
+	CEF3Utils::UnloadCEF3Modules();
+#endif
+}
+
+IWebBrowserSingleton* FWebBrowserModule::GetSingleton()
+{
+	return WebBrowserSingleton;
+}
\ No newline at end of file
diff --git a/Engine/Source/Runtime/WebBrowser/Private/WebBrowserPrivatePCH.h b/Engine/Source/Runtime/WebBrowser/Private/WebBrowserPrivatePCH.h
new file mode 100644
--- /dev/null
+++ b/Engine/Source/Runtime/WebBrowser/Private/WebBrowserPrivatePCH.h
@@ -0,0 +1,7 @@
+// Copyright 1998-2015 Epic Games, Inc. All Rights Reserved.
+
+#pragma once
+
+#include "WebBrowserModule.h"
+
+DECLARE_LOG_CATEGORY_EXTERN(LogWebBrowser, Log, All);
\ No newline at end of file
diff --git a/Engine/Source/Runtime/WebBrowser/Private/WebBrowserSingleton.cpp b/Engine/Source/Runtime/WebBrowser/Private/WebBrowserSingleton.cpp
new file mode 100644
--- /dev/null
+++ b/Engine/Source/Runtime/WebBrowser/Private/WebBrowserSingleton.cpp
@@ -0,0 +1,152 @@
+// Copyright 1998-2015 Epic Games, Inc. All Rights Reserved.
+
+#include "WebBrowserPrivatePCH.h"
+#include "WebBrowserSingleton.h"
+#include "WebBrowserApp.h"
+#include "WebBrowserHandler.h"
+#include "WebBrowserWindow.h"
+
+#if WITH_CEF3
+#if PLATFORM_WINDOWS
+	#include "AllowWindowsPlatformTypes.h"
+#endif
+	#include "include/cef_app.h"
+#if PLATFORM_WINDOWS
+	#include "HideWindowsPlatformTypes.h"
+#endif
+#endif
+
+FWebBrowserSingleton::FWebBrowserSingleton()
+{
+#if WITH_CEF3
+	// Provide CEF with command-line arguments.
+#if PLATFORM_WINDOWS
+	CefMainArgs MainArgs(hInstance);
+#elif PLATFORM_MAC || PLATFORM_LINUX
+    //TArray<FString> Args;
+    //int ArgCount = GSavedCommandLine.ParseIntoArray(&Args, TEXT(" "), true);
+    CefMainArgs MainArgs(0, nullptr);
+#endif
+
+	// WebBrowserApp implements application-level callbacks.
+	WebBrowserApp = new FWebBrowserApp;
+
+	// Specify CEF global settings here.
+	CefSettings Settings;
+	Settings.no_sandbox = true;
+	Settings.command_line_args_disabled = true;
+
+	// Specify locale from our settings
+	FString LocaleCode = GetCurrentLocaleCode();
+	CefString(&Settings.locale) = *LocaleCode;
+
+	// Specify path to resources
+#if PLATFORM_WINDOWS
+	FString ResourcesPath(FPaths::Combine(*FPaths::EngineDir(), TEXT("Binaries/ThirdParty/CEF3/Win64/Resources")));
+	FString LocalesPath(FPaths::Combine(*FPaths::EngineDir(), TEXT("Binaries/ThirdParty/CEF3/Win64/Resources/locales")));
+#elif PLATFORM_MAC
+	FString ResourcesPath(FPaths::Combine(*FPaths::EngineDir(), TEXT("Binaries/ThirdParty/CEF3/Mac/Chromium Embedded Framework.framework/Resources")));
+#elif PLATFORM_LINUX // @todo Linux
+	FString ResourcesPath(FPaths::Combine(*FPaths::EngineDir(), TEXT("Binaries/ThirdParty/CEF3/Linux/Resources")));
+	FString LocalesPath(FPaths::Combine(*FPaths::EngineDir(), TEXT("Binaries/ThirdParty/CEF3/Linux/Resources/locales")));
+#endif
+#if !PLATFORM_MAC // On Mac Chromium ignores custom locales dir. Files need to be stored in Resources folder in the app bundle
+	CefString(&Settings.locales_dir_path) = *LocalesPath;
+#endif
+	CefString(&Settings.resources_dir_path) = *ResourcesPath;
+
+	// Specify path to sub process exe
+#if PLATFORM_WINDOWS
+	FString SubProcessPath(FPaths::Combine(*FPaths::EngineDir(), TEXT("Binaries/Win64/UnrealCEFSubProcess.exe")));
+#elif PLATFORM_MAC
+	FString SubProcessPath(FPaths::Combine(*FPaths::EngineDir(), TEXT("Binaries/Mac/UnrealCEFSubProcess.app/Contents/MacOS/UnrealCEFSubProcess")));
+#else // @todo Linux
+	FString SubProcessPath(TEXT("UnrealCEFSubProcess"));
+#endif
+	CefString(&Settings.browser_subprocess_path) = *SubProcessPath;
+
+	// Initialize CEF.
+	bool bSuccess = CefInitialize(MainArgs, Settings, WebBrowserApp.get(), nullptr);
+	check(bSuccess);
+#endif
+}
+
+FWebBrowserSingleton::~FWebBrowserSingleton()
+{
+#if WITH_CEF3
+	// Force all existing browsers to close in case any haven't been deleted
+	for (int32 Index = 0; Index < WindowInterfaces.Num(); ++Index)
+	{
+		if (WindowInterfaces[Index].IsValid())
+		{
+			WindowInterfaces[Index].Pin()->CloseBrowser();
+		}
+	}
+
+	// CefRefPtr takes care of delete
+	WebBrowserApp = nullptr;
+	// Shut down CEF.
+	CefShutdown();
+#endif
+}
+
+TSharedPtr<IWebBrowserWindow> FWebBrowserSingleton::CreateBrowserWindow(void* OSWindowHandle, FString InitialURL, uint32 Width, uint32 Height, bool bUseTransparency)
+{
+#if WITH_CEF3
+	// Create new window
+	TSharedPtr<FWebBrowserWindow> NewWindow(new FWebBrowserWindow(FIntPoint(Width, Height)));
+
+	// WebBrowserHandler implements browser-level callbacks.
+	CefRefPtr<FWebBrowserHandler> NewHandler(new FWebBrowserHandler);
+	NewWindow->SetHandler(NewHandler);
+
+	// Information used when creating the native window.
+	CefWindowHandle WindowHandle = (CefWindowHandle)OSWindowHandle; // TODO: check this is correct for all platforms
+	CefWindowInfo WindowInfo;
+
+	// Always use off screen rendering so we can integrate with our windows
+	WindowInfo.SetAsOffScreen(WindowHandle);
+	WindowInfo.SetTransparentPainting(bUseTransparency);
+
+	// Specify CEF browser settings here.
+	CefBrowserSettings BrowserSettings;
+
+	CefString URL = *InitialURL;
+
+	// Create the CEF browser window.
+	if (CefBrowserHost::CreateBrowser(WindowInfo, NewHandler.get(), URL, BrowserSettings, NULL))
+	{
+		WindowInterfaces.Add(NewWindow);
+		return NewWindow;
+	}
+#endif
+	return NULL;
+}
+
+bool FWebBrowserSingleton::Tick(float DeltaTime)
+{
+#if WITH_CEF3
+	// Remove any windows that have been deleted
+	for (int32 Index = WindowInterfaces.Num() - 1; Index >= 0; --Index)
+	{
+		if (!WindowInterfaces[Index].IsValid())
+		{
+			WindowInterfaces.RemoveAtSwap(Index);
+		}
+	}
+	CefDoMessageLoopWork();
+#endif
+	return true;
+}
+
+FString FWebBrowserSingleton::GetCurrentLocaleCode()
+{
+	FCultureRef Culture = FInternationalization::Get().GetCurrentCulture();
+	FString LocaleCode = Culture->GetTwoLetterISOLanguageName();
+	FString Country = Culture->GetRegion();
+	if (!Country.IsEmpty())
+	{
+		LocaleCode = LocaleCode + TEXT("-") + Country;
+	}
+	return LocaleCode;
+}
diff --git a/Engine/Source/Runtime/WebBrowser/Private/WebBrowserSingleton.h b/Engine/Source/Runtime/WebBrowser/Private/WebBrowserSingleton.h
new file mode 100644
--- /dev/null
+++ b/Engine/Source/Runtime/WebBrowser/Private/WebBrowserSingleton.h
@@ -0,0 +1,57 @@
+// Copyright 1998-2015 Epic Games, Inc. All Rights Reserved.
+
+#pragma once
+
+#include "IWebBrowserSingleton.h"
+
+#if WITH_CEF3
+#if PLATFORM_WINDOWS
+#include "AllowWindowsPlatformTypes.h"
+#endif
+#include "include/internal/cef_ptr.h"
+#if PLATFORM_WINDOWS
+#include "HideWindowsPlatformTypes.h"
+#endif
+#endif
+
+// Forward Declarations
+class FWebBrowserApp;
+class FWebBrowserHandler;
+class FWebBrowserWindow;
+
+/**
+ * Implementation of singleton class that takes care of general web browser tasks
+ */
+class FWebBrowserSingleton : public IWebBrowserSingleton, public FTickerObjectBase
+{
+public:
+	/**
+	 * Default Constructor
+	 */
+	FWebBrowserSingleton();
+	/**
+	 * Virtual Destructor
+	 */
+	virtual ~FWebBrowserSingleton();
+
+	// IWebBrowserSingleton Interface
+	TSharedPtr<IWebBrowserWindow> CreateBrowserWindow(void* OSWindowHandle, FString InitialURL, uint32 Width, uint32 Height, bool bUseTransparency) override;
+
+	// FTickerObjectBase Interface
+	virtual bool Tick(float DeltaTime) override;
+
+	/**
+	 * Gets the Current Locale Code in the format CEF expects
+	 *
+	 * @return Locale code as either "xx" or "xx-YY"
+	 */
+	static FString GetCurrentLocaleCode();
+
+private:
+#if WITH_CEF3
+	/** Pointer to the CEF App implementation */
+	CefRefPtr<FWebBrowserApp>			WebBrowserApp;
+	/** List of currently existing browser windows */
+	TArray<TWeakPtr<FWebBrowserWindow>>	WindowInterfaces;
+#endif
+};
diff --git a/Engine/Source/Runtime/WebBrowser/Private/WebBrowserViewport.cpp b/Engine/Source/Runtime/WebBrowser/Private/WebBrowserViewport.cpp
new file mode 100644
--- /dev/null
+++ b/Engine/Source/Runtime/WebBrowser/Private/WebBrowserViewport.cpp
@@ -0,0 +1,110 @@
+// Copyright 1998-2015 Epic Games, Inc. All Rights Reserved.
+
+#include "WebBrowserPrivatePCH.h"
+#include "WebBrowserViewport.h"
+#include "IWebBrowserWindow.h"
+
+FIntPoint FWebBrowserViewport::GetSize() const
+{
+	return (WebBrowserWindow->GetTexture() != nullptr)
+		? FIntPoint(WebBrowserWindow->GetTexture()->GetWidth(), WebBrowserWindow->GetTexture()->GetHeight())
+		: FIntPoint();
+}
+
+FSlateShaderResource* FWebBrowserViewport::GetViewportRenderTargetTexture() const
+{
+	return WebBrowserWindow->GetTexture();
+}
+
+void FWebBrowserViewport::Tick(const FGeometry& AllottedGeometry, double InCurrentTime, float DeltaTime)
+{
+	WebBrowserWindow->SetViewportSize(AllottedGeometry.GetLocalSize());
+}
+
+bool FWebBrowserViewport::RequiresVsync() const
+{
+	return false;
+}
+
+FCursorReply FWebBrowserViewport::OnCursorQuery( const FGeometry& MyGeometry, const FPointerEvent& CursorEvent )
+{
+	// TODO: retrieve cursor type from WebBrowserWindow if we can figure that out from CefCursorHandle
+	return FCursorReply::Unhandled();
+}
+
+FReply FWebBrowserViewport::OnMouseButtonDown(const FGeometry& MyGeometry, const FPointerEvent& MouseEvent)
+{
+	// Capture mouse on left button down so that you can drag out of the viewport
+	WebBrowserWindow->OnMouseButtonDown(MyGeometry, MouseEvent);
+	if (MouseEvent.GetEffectingButton() == EKeys::LeftMouseButton && ViewportWidget.IsValid())
+	{
+		return FReply::Handled().CaptureMouse(ViewportWidget.Pin().ToSharedRef());
+	}
+	return FReply::Handled();
+}
+
+FReply FWebBrowserViewport::OnMouseButtonUp(const FGeometry& MyGeometry, const FPointerEvent& MouseEvent)
+{
+	// Release mouse capture when left button released
+	WebBrowserWindow->OnMouseButtonUp(MyGeometry, MouseEvent);
+	if (MouseEvent.GetEffectingButton() == EKeys::LeftMouseButton)
+	{
+		return FReply::Handled().ReleaseMouseCapture();
+	}
+	return FReply::Handled();
+}
+
+void FWebBrowserViewport::OnMouseEnter(const FGeometry& MyGeometry, const FPointerEvent& MouseEvent)
+{
+}
+
+void FWebBrowserViewport::OnMouseLeave(const FPointerEvent& MouseEvent)
+{
+}
+
+FReply FWebBrowserViewport::OnMouseMove(const FGeometry& MyGeometry, const FPointerEvent& MouseEvent)
+{
+	WebBrowserWindow->OnMouseMove(MyGeometry, MouseEvent);
+	return FReply::Handled();
+}
+
+FReply FWebBrowserViewport::OnMouseWheel(const FGeometry& MyGeometry, const FPointerEvent& MouseEvent)
+{
+	WebBrowserWindow->OnMouseWheel(MyGeometry, MouseEvent);
+	return FReply::Handled();
+}
+
+FReply FWebBrowserViewport::OnMouseButtonDoubleClick(const FGeometry& InMyGeometry, const FPointerEvent& InMouseEvent)
+{
+	WebBrowserWindow->OnMouseButtonDoubleClick(InMyGeometry, InMouseEvent);
+	return FReply::Handled();
+}
+
+FReply FWebBrowserViewport::OnKeyDown(const FGeometry& MyGeometry, const FKeyEvent& InKeyEvent)
+{
+	WebBrowserWindow->OnKeyDown(InKeyEvent);
+	return FReply::Handled();
+}
+
+FReply FWebBrowserViewport::OnKeyUp(const FGeometry& MyGeometry, const FKeyEvent& InKeyEvent)
+{
+	WebBrowserWindow->OnKeyUp(InKeyEvent);
+	return FReply::Handled();
+}
+
+FReply FWebBrowserViewport::OnKeyChar( const FGeometry& MyGeometry, const FCharacterEvent& InCharacterEvent )
+{
+	WebBrowserWindow->OnKeyChar(InCharacterEvent);
+	return FReply::Handled();
+}
+
+FReply FWebBrowserViewport::OnFocusReceived(const FFocusEvent& InFocusEvent)
+{
+	WebBrowserWindow->OnFocus(true);
+	return FReply::Handled();
+}
+
+void FWebBrowserViewport::OnFocusLost(const FFocusEvent& InFocusEvent)
+{
+	WebBrowserWindow->OnFocus(false);
+}
diff --git a/Engine/Source/Runtime/WebBrowser/Private/WebBrowserWindow.cpp b/Engine/Source/Runtime/WebBrowser/Private/WebBrowserWindow.cpp
new file mode 100644
--- /dev/null
+++ b/Engine/Source/Runtime/WebBrowser/Private/WebBrowserWindow.cpp
@@ -0,0 +1,479 @@
+// Copyright 1998-2015 Epic Games, Inc. All Rights Reserved.
+
+#include "WebBrowserPrivatePCH.h"
+#include "WebBrowserWindow.h"
+#include "SlateCore.h"
+#include "SlateBasics.h"
+
+#if WITH_CEF3
+FWebBrowserWindow::FWebBrowserWindow(FIntPoint InViewportSize)
+	: UpdatableTexture(nullptr)
+	, ViewportSize(InViewportSize)
+	, bIsClosing(false)
+	, bHasBeenPainted(false)
+{
+	TextureData.Reserve(ViewportSize.X * ViewportSize.Y * 4);
+	TextureData.SetNumZeroed(ViewportSize.X * ViewportSize.Y * 4);
+
+	if (FSlateApplication::Get().GetRenderer().IsValid())
+	{
+		UpdatableTexture = FSlateApplication::Get().GetRenderer()->CreateUpdatableTexture(ViewportSize.X, ViewportSize.Y);
+	}
+}
+
+FWebBrowserWindow::~FWebBrowserWindow()
+{
+	CloseBrowser();
+
+	if (FSlateApplication::IsInitialized() && FSlateApplication::Get().GetRenderer().IsValid() && UpdatableTexture != nullptr)
+	{
+		FSlateApplication::Get().GetRenderer()->ReleaseUpdatableTexture(UpdatableTexture);
+	}
+	UpdatableTexture = nullptr;
+}
+
+void FWebBrowserWindow::SetViewportSize(FVector2D WindowSize)
+{
+	// Magic number for texture size, can't access GetMax2DTextureDimension easily
+	FIntPoint ClampedWindowSize = WindowSize.ClampAxes(1, 2048).IntPoint();
+
+	// Ignore sizes that can't be seen as it forces CEF to re-render whole image
+	if (WindowSize >= FVector2D::UnitVector && ViewportSize != ClampedWindowSize)
+	{
+		FIntPoint OldViewportSize = MoveTemp(ViewportSize);
+		TArray<uint8> OldTextureData = MoveTemp(TextureData);
+		ViewportSize = ClampedWindowSize;
+		TextureData.SetNumZeroed(ViewportSize.X * ViewportSize.Y * 4);
+
+		// copy row by row to avoid texture distortion
+		const int32 WriteWidth = FMath::Min(OldViewportSize.X, ViewportSize.X) * 4;
+		const int32 WriteHeight = FMath::Min(OldViewportSize.Y, ViewportSize.Y);
+		for (int32 RowIndex = 0; RowIndex < WriteHeight; ++RowIndex)
+		{
+			FMemory::Memcpy(TextureData.GetData() + ViewportSize.X * RowIndex * 4, OldTextureData.GetData() + OldViewportSize.X * RowIndex * 4, WriteWidth);
+		}
+
+		if (UpdatableTexture != nullptr)
+		{
+			UpdatableTexture->ResizeTexture(ViewportSize.X, ViewportSize.Y);
+			UpdatableTexture->UpdateTexture(TextureData);
+		}
+		if (IsValid())
+		{
+			InternalCefBrowser->GetHost()->WasResized();
+		}
+	}
+}
+
+FSlateShaderResource* FWebBrowserWindow::GetTexture()
+{
+	if (UpdatableTexture != nullptr)
+	{
+		return UpdatableTexture->GetSlateResource();
+	}
+	return nullptr;
+}
+
+bool FWebBrowserWindow::IsValid() const
+{
+	return InternalCefBrowser.get() != nullptr;
+}
+
+bool FWebBrowserWindow::HasBeenPainted() const
+{
+	return bHasBeenPainted;
+}
+
+bool FWebBrowserWindow::IsClosing() const
+{
+	return bIsClosing;
+}
+
+FString FWebBrowserWindow::GetTitle() const
+{
+	return Title;
+}
+
+void FWebBrowserWindow::OnKeyDown(const FKeyEvent& InKeyEvent)
+{
+	if (IsValid())
+	{
+		CefKeyEvent KeyEvent;
+		KeyEvent.windows_key_code = InKeyEvent.GetKeyCode();
+		// TODO: Figure out whether this is a system key if we come across problems
+		/*KeyEvent.is_system_key = message == WM_SYSCHAR ||
+			message == WM_SYSKEYDOWN ||
+			message == WM_SYSKEYUP;*/
+
+		KeyEvent.type = KEYEVENT_RAWKEYDOWN;
+		KeyEvent.modifiers = GetCefKeyboardModifiers(InKeyEvent);
+
+		InternalCefBrowser->GetHost()->SendKeyEvent(KeyEvent);
+	}
+}
+
+void FWebBrowserWindow::OnKeyUp(const FKeyEvent& InKeyEvent)
+{
+	if (IsValid())
+	{
+		CefKeyEvent KeyEvent;
+		KeyEvent.windows_key_code = InKeyEvent.GetKeyCode();
+		// TODO: Figure out whether this is a system key if we come across problems
+		/*KeyEvent.is_system_key = message == WM_SYSCHAR ||
+			message == WM_SYSKEYDOWN ||
+			message == WM_SYSKEYUP;*/
+
+		KeyEvent.type = KEYEVENT_KEYUP;
+		KeyEvent.modifiers = GetCefKeyboardModifiers(InKeyEvent);
+
+		InternalCefBrowser->GetHost()->SendKeyEvent(KeyEvent);
+	}
+}
+
+void FWebBrowserWindow::OnKeyChar(const FCharacterEvent& InCharacterEvent)
+{
+	if (IsValid())
+	{
+		CefKeyEvent KeyEvent;
+		KeyEvent.windows_key_code = InCharacterEvent.GetCharacter();
+		// TODO: Figure out whether this is a system key if we come across problems
+		/*KeyEvent.is_system_key = message == WM_SYSCHAR ||
+			message == WM_SYSKEYDOWN ||
+			message == WM_SYSKEYUP;*/
+
+		KeyEvent.type = KEYEVENT_CHAR;
+		KeyEvent.modifiers = GetCefInputModifiers(InCharacterEvent);
+
+		InternalCefBrowser->GetHost()->SendKeyEvent(KeyEvent);
+	}
+}
+
+void FWebBrowserWindow::OnMouseButtonDown(const FGeometry& MyGeometry, const FPointerEvent& MouseEvent)
+{
+	if (IsValid())
+	{
+		FKey Button = MouseEvent.GetEffectingButton();
+		CefBrowserHost::MouseButtonType Type =
+			(Button == EKeys::LeftMouseButton ? MBT_LEFT : (
+			Button == EKeys::RightMouseButton ? MBT_RIGHT : MBT_MIDDLE));
+
+		CefMouseEvent Event;
+		FVector2D LocalPos = MyGeometry.AbsoluteToLocal(MouseEvent.GetScreenSpacePosition());
+		Event.x = LocalPos.X;
+		Event.y = LocalPos.Y;
+		Event.modifiers = GetCefMouseModifiers(MouseEvent);
+
+		InternalCefBrowser->GetHost()->SendMouseClickEvent(Event, Type, false,1);
+	}
+}
+
+void FWebBrowserWindow::OnMouseButtonUp(const FGeometry& MyGeometry, const FPointerEvent& MouseEvent)
+{
+	if (IsValid())
+	{
+		FKey Button = MouseEvent.GetEffectingButton();
+		CefBrowserHost::MouseButtonType Type =
+			(Button == EKeys::LeftMouseButton ? MBT_LEFT : (
+			Button == EKeys::RightMouseButton ? MBT_RIGHT : MBT_MIDDLE));
+
+		CefMouseEvent Event;
+		FVector2D LocalPos = MyGeometry.AbsoluteToLocal(MouseEvent.GetScreenSpacePosition());
+		Event.x = LocalPos.X;
+		Event.y = LocalPos.Y;
+		Event.modifiers = GetCefMouseModifiers(MouseEvent);
+
+		InternalCefBrowser->GetHost()->SendMouseClickEvent(Event, Type, true, 1);
+	}
+}
+
+void FWebBrowserWindow::OnMouseButtonDoubleClick(const FGeometry& MyGeometry, const FPointerEvent& MouseEvent)
+{
+	if (IsValid())
+	{
+		FKey Button = MouseEvent.GetEffectingButton();
+		CefBrowserHost::MouseButtonType Type =
+			(Button == EKeys::LeftMouseButton ? MBT_LEFT : (
+			Button == EKeys::RightMouseButton ? MBT_RIGHT : MBT_MIDDLE));
+
+		CefMouseEvent Event;
+		FVector2D LocalPos = MyGeometry.AbsoluteToLocal(MouseEvent.GetScreenSpacePosition());
+		Event.x = LocalPos.X;
+		Event.y = LocalPos.Y;
+		Event.modifiers = GetCefMouseModifiers(MouseEvent);
+
+		InternalCefBrowser->GetHost()->SendMouseClickEvent(Event, Type, false, 2);
+	}
+}
+
+void FWebBrowserWindow::OnMouseMove(const FGeometry& MyGeometry, const FPointerEvent& MouseEvent)
+{
+	if (IsValid())
+	{
+		CefMouseEvent Event;
+		FVector2D LocalPos = MyGeometry.AbsoluteToLocal(MouseEvent.GetScreenSpacePosition());
+		Event.x = LocalPos.X;
+		Event.y = LocalPos.Y;
+		Event.modifiers = GetCefMouseModifiers(MouseEvent);
+
+		InternalCefBrowser->GetHost()->SendMouseMoveEvent(Event, false);
+	}
+}
+
+void FWebBrowserWindow::OnMouseWheel(const FGeometry& MyGeometry, const FPointerEvent& MouseEvent)
+{
+	if(IsValid())
+	{
+		// The original delta is reduced so this should bring it back to what CEF expects
+		const float SpinFactor = 120.0f;
+		const float TrueDelta = MouseEvent.GetWheelDelta() * SpinFactor;
+		CefMouseEvent Event;
+		FVector2D LocalPos = MyGeometry.AbsoluteToLocal(MouseEvent.GetScreenSpacePosition());
+		Event.x = LocalPos.X;
+		Event.y = LocalPos.Y;
+		Event.modifiers = GetCefMouseModifiers(MouseEvent);
+		
+		InternalCefBrowser->GetHost()->SendMouseWheelEvent(Event,
+															MouseEvent.IsShiftDown() ? TrueDelta : 0,
+															!MouseEvent.IsShiftDown() ? TrueDelta : 0);
+	}
+}
+
+void FWebBrowserWindow::OnFocus(bool SetFocus)
+{
+	if (IsValid())
+	{
+		InternalCefBrowser->GetHost()->SendFocusEvent(SetFocus);
+	}
+}
+
+void FWebBrowserWindow::OnCaptureLost()
+{
+	if (IsValid())
+	{
+		InternalCefBrowser->GetHost()->SendCaptureLostEvent();
+	}
+}
+
+bool FWebBrowserWindow::CanGoBack() const
+{
+	if (IsValid())
+	{
+		return InternalCefBrowser->CanGoBack();
+	}
+	return false;
+}
+
+void FWebBrowserWindow::GoBack()
+{
+	if (IsValid())
+	{
+		InternalCefBrowser->GoBack();
+	}
+}
+
+bool FWebBrowserWindow::CanGoForward() const
+{
+	if (IsValid())
+	{
+		return InternalCefBrowser->CanGoForward();
+	}
+	return false;
+}
+
+void FWebBrowserWindow::GoForward()
+{
+	if (IsValid())
+	{
+		InternalCefBrowser->GoForward();
+	}
+}
+
+bool FWebBrowserWindow::IsLoading() const
+{
+	if (IsValid())
+	{
+		return InternalCefBrowser->IsLoading();
+	}
+	return false;
+}
+
+void FWebBrowserWindow::Reload()
+{
+	if (IsValid())
+	{
+		InternalCefBrowser->Reload();
+	}
+}
+
+void FWebBrowserWindow::StopLoad()
+{
+	if (IsValid())
+	{
+		InternalCefBrowser->StopLoad();
+	}
+}
+
+void FWebBrowserWindow::SetHandler(CefRefPtr<FWebBrowserHandler> InHandler)
+{
+	if (InHandler.get())
+	{
+		Handler = InHandler;
+		Handler->SetBrowserWindow(SharedThis(this));
+	}
+}
+
+void FWebBrowserWindow::CloseBrowser()
+{
+	bIsClosing = true;
+	if (IsValid())
+	{
+		InternalCefBrowser->GetHost()->CloseBrowser(false);
+		InternalCefBrowser = nullptr;
+		Handler = nullptr;
+	}
+}
+
+void FWebBrowserWindow::BindCefBrowser(CefRefPtr<CefBrowser> Browser)
+{
+	InternalCefBrowser = Browser;
+}
+
+void FWebBrowserWindow::SetTitle(const CefString& InTitle)
+{
+	Title = InTitle.ToWString().c_str();
+	OnTitleChangedDelegate.Broadcast(Title);
+}
+
+bool FWebBrowserWindow::GetViewRect(CefRect& Rect)
+{
+	Rect.x = 0;
+	Rect.y = 0;
+	Rect.width = ViewportSize.X;
+	Rect.height = ViewportSize.Y;
+
+	return true;
+}
+
+void FWebBrowserWindow::OnPaint(CefRenderHandler::PaintElementType Type, const CefRenderHandler::RectList& DirtyRects, const void* Buffer, int Width, int Height)
+{
+	const int32 BufferSize = Width*Height*4;
+	if (BufferSize == TextureData.Num())
+	{
+		FMemory::Memcpy(TextureData.GetData(), Buffer, BufferSize);
+	}
+	else
+	{
+		// copy row by row to avoid texture distortion
+		const int32 WriteWidth = FMath::Min(Width, ViewportSize.X) * 4;
+		const int32 WriteHeight = FMath::Min(Height, ViewportSize.Y);
+		for (int32 RowIndex = 0; RowIndex < WriteHeight; ++RowIndex)
+		{
+			FMemory::Memcpy(TextureData.GetData() + ViewportSize.X * RowIndex * 4, static_cast<const uint8*>(Buffer) + Width * RowIndex * 4, WriteWidth);
+		}
+	}
+
+	if (UpdatableTexture != nullptr)
+	{
+		UpdatableTexture->UpdateTexture(TextureData);
+	}
+
+	bHasBeenPainted = true;
+}
+
+void FWebBrowserWindow::OnCursorChange(CefCursorHandle Cursor)
+{
+	// TODO: Figure out Unreal cursor type from this,
+	// may need to reload unreal cursors to compare handles
+	//::SetCursor( Cursor );
+}
+
+int32 FWebBrowserWindow::GetCefKeyboardModifiers(const FKeyEvent& KeyEvent)
+{
+	int32 Modifiers = GetCefInputModifiers(KeyEvent);
+
+	const FKey Key = KeyEvent.GetKey();
+	if (Key == EKeys::LeftAlt ||
+		Key == EKeys::LeftCommand ||
+		Key == EKeys::LeftControl ||
+		Key == EKeys::LeftShift)
+	{
+		Modifiers |= EVENTFLAG_IS_LEFT;
+	}
+	if (Key == EKeys::RightAlt ||
+		Key == EKeys::RightCommand ||
+		Key == EKeys::RightControl ||
+		Key == EKeys::RightShift)
+	{
+		Modifiers |= EVENTFLAG_IS_RIGHT;
+	}
+	if (Key == EKeys::NumPadZero ||
+		Key == EKeys::NumPadOne ||
+		Key == EKeys::NumPadTwo ||
+		Key == EKeys::NumPadThree ||
+		Key == EKeys::NumPadFour ||
+		Key == EKeys::NumPadFive ||
+		Key == EKeys::NumPadSix ||
+		Key == EKeys::NumPadSeven ||
+		Key == EKeys::NumPadEight ||
+		Key == EKeys::NumPadNine)
+	{
+		Modifiers |= EVENTFLAG_IS_KEY_PAD;
+	}
+
+	return Modifiers;
+}
+
+int32 FWebBrowserWindow::GetCefMouseModifiers(const FPointerEvent& InMouseEvent)
+{
+	int32 Modifiers = GetCefInputModifiers(InMouseEvent);
+
+	if (InMouseEvent.IsMouseButtonDown(EKeys::LeftMouseButton))
+	{
+		Modifiers |= EVENTFLAG_LEFT_MOUSE_BUTTON;
+	}
+	if (InMouseEvent.IsMouseButtonDown(EKeys::MiddleMouseButton))
+	{
+		Modifiers |= EVENTFLAG_MIDDLE_MOUSE_BUTTON;
+	}
+	if (InMouseEvent.IsMouseButtonDown(EKeys::RightMouseButton))
+	{
+		Modifiers |= EVENTFLAG_RIGHT_MOUSE_BUTTON;
+	}
+
+	return Modifiers;
+}
+
+int32 FWebBrowserWindow::GetCefInputModifiers(const FInputEvent& InputEvent)
+{
+	int32 Modifiers = 0;
+
+	if (InputEvent.IsShiftDown())
+	{
+		Modifiers |= EVENTFLAG_SHIFT_DOWN;
+	}
+	if (InputEvent.IsControlDown())
+	{
+		Modifiers |= EVENTFLAG_CONTROL_DOWN;
+	}
+	if (InputEvent.IsAltDown())
+	{
+		Modifiers |= EVENTFLAG_ALT_DOWN;
+	}
+	if (InputEvent.IsCommandDown())
+	{
+		Modifiers |= EVENTFLAG_COMMAND_DOWN;
+	}
+	if (InputEvent.AreCapsLocked())
+	{
+		Modifiers |= EVENTFLAG_CAPS_LOCK_ON;
+	}
+	// TODO: Add function for this if necessary
+	/*if (InputEvent.AreNumsLocked())
+	{
+		Modifiers |= EVENTFLAG_NUM_LOCK_ON;
+	}*/
+
+	return Modifiers;
+}
+
+#endif
diff --git a/Engine/Source/Runtime/WebBrowser/Private/WebBrowserWindow.h b/Engine/Source/Runtime/WebBrowser/Private/WebBrowserWindow.h
new file mode 100644
--- /dev/null
+++ b/Engine/Source/Runtime/WebBrowser/Private/WebBrowserWindow.h
@@ -0,0 +1,166 @@
+// Copyright 1998-2015 Epic Games, Inc. All Rights Reserved.
+
+#pragma once
+
+#if WITH_CEF3
+
+#include "IWebBrowserWindow.h"
+#include "WebBrowserHandler.h"
+
+#if PLATFORM_WINDOWS
+#include "AllowWindowsPlatformTypes.h"
+#endif
+#include "include/internal/cef_ptr.h"
+#include "include/cef_render_handler.h"
+#if PLATFORM_WINDOWS
+#include "HideWindowsPlatformTypes.h"
+#endif
+
+// Forward Declarations
+struct FInputEvent;
+
+class FSlateRenderer;
+class FSlateUpdatableTexture;
+
+/**
+ * Implementation of interface for dealing with a Web Browser window
+ */
+class FWebBrowserWindow : public IWebBrowserWindow, public TSharedFromThis<FWebBrowserWindow>
+{
+public:
+	/**
+	 * Default Constructor
+	 * 
+	 * @param InViewportSize Initial size of the browser window
+	 */
+	FWebBrowserWindow(FIntPoint InViewportSize);
+	/**
+	 * Virtual Destructor
+	 */
+	virtual ~FWebBrowserWindow();
+
+	// IWebBrowserWindow Interface
+	virtual void SetViewportSize(FVector2D WindowSize) override;
+	virtual FSlateShaderResource* GetTexture() override;
+	virtual bool IsValid() const override;
+	virtual bool HasBeenPainted() const override;
+	virtual bool IsClosing() const override;
+	virtual FString GetTitle() const override;
+	virtual void OnKeyDown(const FKeyEvent& InKeyEvent) override;
+	virtual void OnKeyUp(const FKeyEvent& InKeyEvent) override;
+	virtual void OnKeyChar(const FCharacterEvent& InCharacterEvent) override;
+	virtual void OnMouseButtonDown(const FGeometry& MyGeometry, const FPointerEvent& MouseEvent) override;
+	virtual void OnMouseButtonUp(const FGeometry& MyGeometry, const FPointerEvent& MouseEvent) override;
+	virtual void OnMouseButtonDoubleClick(const FGeometry& MyGeometry, const FPointerEvent& MouseEvent) override;
+	virtual void OnMouseMove(const FGeometry& MyGeometry, const FPointerEvent& MouseEvent) override;
+	virtual void OnMouseWheel(const FGeometry& MyGeometry, const FPointerEvent& MouseEvent) override;
+	virtual void OnFocus(bool SetFocus) override;
+	virtual void OnCaptureLost() override;
+	virtual bool CanGoBack() const override;
+	virtual void GoBack() override;
+	virtual bool CanGoForward() const override;
+	virtual void GoForward() override;
+	virtual bool IsLoading() const override;
+	virtual void Reload() override;
+	virtual void StopLoad() override;
+
+	/**
+	 * Set the CEF Handler receiving browser callbacks for this window
+	 *
+	 * @param InHandler Pointer to the handler for this window
+	 */
+	void SetHandler(CefRefPtr<FWebBrowserHandler> InHandler);
+
+	/**
+	 * Close this window so that it can no longer be used
+	 */
+	void CloseBrowser();
+
+private:
+	/**
+	 * Called to pass reference to the underlying CefBrowser as this is not created at the same time
+	 * as the FWebBrowserWindow
+	 *
+	 * @param Browser The CefBrowser for this window
+	 */
+	void BindCefBrowser(CefRefPtr<CefBrowser> Browser);
+
+	/**
+	 * Sets the Title of this window
+	 * 
+	 * @param InTitle The new title of this window
+	 */
+	void SetTitle(const CefString& InTitle);
+
+	/**
+	 * Get the current proportions of this window
+	 *
+	 * @param Rect Reference to CefRect to store sizes
+	 * @return Whether Rect was set up correctly
+	 */
+	bool GetViewRect(CefRect& Rect);
+
+	/**
+	 * Called when there is an update to the rendered web page
+	 *
+	 * @param Type Paint type
+	 * @param DirtyRects List of image areas that have been changed
+	 * @param Buffer Pointer to the raw texture data
+	 * @param Width Width of the texture
+	 * @param Height Height of the texture
+	 */
+	void OnPaint(CefRenderHandler::PaintElementType Type, const CefRenderHandler::RectList& DirtyRects, const void* Buffer, int Width, int Height);
+	
+	/**
+	 * Called when cursor would change due to web browser interaction
+	 * 
+	 * @param Cursor Handle to CEF mouse cursor
+	 */
+	void OnCursorChange(CefCursorHandle Cursor);
+
+	/**
+	 * Gets the Cef Keyboard Modifiers based on a Key Event
+	 * 
+	 * @param KeyEvent The Key event
+	 * @return Bits representing keyboard modifiers
+	 */
+	static int32 GetCefKeyboardModifiers(const FKeyEvent& KeyEvent);
+
+	/**
+	 * Gets the Cef Mouse Modifiers based on a Mouse Event
+	 * 
+	 * @param InMouseEvent The Mouse event
+	 * @return Bits representing mouse modifiers
+	 */
+	static int32 GetCefMouseModifiers(const FPointerEvent& InMouseEvent);
+
+	/**
+	 * Gets the Cef Input Modifiers based on an Input Event
+	 * 
+	 * @param InputEvent The Input event
+	 * @return Bits representing input modifiers
+	 */
+	static int32 GetCefInputModifiers(const FInputEvent& InputEvent);
+
+	/** Interface to the texture we are rendering to */
+	FSlateUpdatableTexture*			UpdatableTexture;
+	/** Temporary storage for the raw texture data */
+	TArray<uint8>					TextureData;
+	/** Pointer to the CEF Handler for this window */
+	CefRefPtr<FWebBrowserHandler>	Handler;
+	/** Pointer to the CEF Browser for this window */
+	CefRefPtr<CefBrowser>			InternalCefBrowser;
+	/** Current title of this window */
+	FString							Title;
+	/** Current size of this window */
+	FIntPoint						ViewportSize;
+	/** Whether this window is closing */
+	bool							bIsClosing;
+	/** Whether this window has been painted at least once */
+	bool							bHasBeenPainted;
+
+	// Allow the Handler to access functions only it needs
+	friend class FWebBrowserHandler;
+};
+
+#endif
diff --git a/Engine/Source/Runtime/WebBrowser/Public/IWebBrowserSingleton.h b/Engine/Source/Runtime/WebBrowser/Public/IWebBrowserSingleton.h
new file mode 100644
--- /dev/null
+++ b/Engine/Source/Runtime/WebBrowser/Public/IWebBrowserSingleton.h
@@ -0,0 +1,31 @@
+// Copyright 1998-2015 Epic Games, Inc. All Rights Reserved.
+
+#pragma once
+
+// Forward Declarations
+class FSlateRenderer;
+class IWebBrowserWindow;
+
+/**
+ * A singleton class that takes care of general web browser tasks
+ */
+class WEBBROWSER_API IWebBrowserSingleton
+{
+public:
+	/**
+	 * Virtual Destructor
+	 */
+	virtual ~IWebBrowserSingleton() {};
+
+	/**
+	 * Create a new web browser window
+	 *
+	 * @param OSWindowHandle Handle of OS Window that the browser will be displayed in (can be null)
+	 * @param InitialURL URL that the browser should initially navigate to
+	 * @param Width Initial width of the browser
+	 * @param Height Initial height of the browser
+	 * @param bUseTransparency Whether to allow transparent rendering of pages
+	 * @return New Web Browser Window Interface (may be null if not supported)
+	 */
+	virtual TSharedPtr<IWebBrowserWindow> CreateBrowserWindow(void* OSWindowHandle, FString InitialURL, uint32 Width, uint32 Height, bool bUseTransparency) = 0;
+};
diff --git a/Engine/Source/Runtime/WebBrowser/Public/IWebBrowserWindow.h b/Engine/Source/Runtime/WebBrowser/Public/IWebBrowserWindow.h
new file mode 100644
--- /dev/null
+++ b/Engine/Source/Runtime/WebBrowser/Public/IWebBrowserWindow.h
@@ -0,0 +1,176 @@
+// Copyright 1998-2015 Epic Games, Inc. All Rights Reserved.
+
+#pragma once
+
+// Forward Declarations
+struct FGeometry;
+struct FKeyEvent;
+struct FCharacterEvent;
+struct FPointerEvent;
+class FSlateShaderResource;
+
+DECLARE_MULTICAST_DELEGATE_OneParam( FOnTitleChangedDelegate, FString );
+
+/**
+ * Interface for dealing with a Web Browser window
+ */
+class IWebBrowserWindow
+{
+public:
+	/**
+	 * Set the desired size of the web browser viewport
+	 * 
+	 * @param WindowSize Desired viewport size
+	 */
+	virtual void SetViewportSize(FVector2D WindowSize) = 0;
+
+	/**
+	 * Gets interface to the texture representation of the browser
+	 *
+	 * @return A slate shader resource that can be rendered
+	 */
+	virtual FSlateShaderResource* GetTexture() = 0;
+
+	/**
+	 * Checks whether the web browser is valid and ready for use
+	 */
+	virtual bool IsValid() const = 0;
+
+	/**
+	 * Checks whether the web browser has been painted at least once
+	 */
+	virtual bool HasBeenPainted() const = 0;
+
+	/**
+	 * Checks whether the web browser is currently being shut down
+	 */
+	virtual bool IsClosing() const = 0;
+
+	/**
+	 * Gets the current title of the Browser page
+	 */
+	virtual FString GetTitle() const = 0;
+
+	/**
+	 * Notify the browser that a key has been pressed
+	 *
+	 * @param  InKeyEvent  Key event
+	 */
+	virtual void OnKeyDown(const FKeyEvent& InKeyEvent) = 0;
+
+	/**
+	 * Notify the browser that a key has been released
+	 *
+	 * @param  InKeyEvent  Key event
+	 */
+	virtual void OnKeyUp(const FKeyEvent& InKeyEvent) = 0;
+
+	/**
+	 * Notify the browser of a character event
+	 *
+	 * @param InCharacterEvent Character event
+	 */
+	virtual void OnKeyChar(const FCharacterEvent& InCharacterEvent) = 0;
+
+	/**
+	 * Notify the browser that a mouse button was pressed within it
+	 *
+	 * @param MyGeometry The Geometry of the browser
+	 * @param MouseEvent Information about the input event
+	 */
+	virtual void OnMouseButtonDown(const FGeometry& MyGeometry, const FPointerEvent& MouseEvent) = 0;
+
+	/**
+	 * Notify the browser that a mouse button was released within it
+	 *
+	 * @param MyGeometry The Geometry of the browser
+	 * @param MouseEvent Information about the input event
+	 */
+	virtual void OnMouseButtonUp(const FGeometry& MyGeometry, const FPointerEvent& MouseEvent) = 0;
+
+	/**
+	 * Notify the browser of a double click event
+	 *
+	 * @param MyGeometry The Geometry of the browser
+	 * @param MouseEvent Information about the input event
+	 */
+	virtual void OnMouseButtonDoubleClick(const FGeometry& MyGeometry, const FPointerEvent& MouseEvent) = 0;
+
+	/**
+	 * Notify the browser that a mouse moved within it
+	 *
+	 * @param MyGeometry The Geometry of the browser
+	 * @param MouseEvent Information about the input event
+	 */
+	virtual void OnMouseMove(const FGeometry& MyGeometry, const FPointerEvent& MouseEvent) = 0;
+
+	/**
+	 * Called when the mouse wheel is spun
+	 *
+	 * @param MyGeometry The Geometry of the browser
+	 * @param MouseEvent Information about the input event
+	 */
+	virtual void OnMouseWheel(const FGeometry& MyGeometry, const FPointerEvent& MouseEvent) = 0;
+
+	/**
+	 * Called when browser receives/loses focus
+	 */
+	virtual void OnFocus(bool SetFocus) = 0;
+
+	/**
+	 * Called when Capture lost
+	 */
+	virtual void OnCaptureLost() = 0;
+
+	/**
+	 * Returns true if the browser can navigate backwards.
+	 */
+	virtual bool CanGoBack() const = 0;
+
+	/**
+	 * Navigate backwards.
+	 */
+	virtual void GoBack() = 0;
+
+	/**
+	 * Returns true if the browser can navigate forwards.
+	 */
+	virtual bool CanGoForward() const = 0;
+
+	/**
+	 * Navigate forwards.
+	 */
+	virtual void GoForward() = 0;
+
+	/**
+	 * Returns true if the browser is currently loading.
+	 */
+	virtual bool IsLoading() const = 0;
+
+	/**
+	 * Reload the current page.
+	 */
+	virtual void Reload() = 0;
+
+	/**
+	 * Stop loading the page.
+	 */
+	virtual void StopLoad() = 0;
+
+	/**
+	 * Accesses a delegate to allow callbacks when a browser title changes
+	 *
+	 * @return A multicast delegate that you can register with
+	 */
+	FOnTitleChangedDelegate& OnTitleChanged() {return OnTitleChangedDelegate;}
+
+protected:
+
+	/**
+	 * Virtual Destructor
+	 */
+	virtual ~IWebBrowserWindow() {};
+
+	/** Delegate for broadcasting title changes */
+	FOnTitleChangedDelegate OnTitleChangedDelegate;
+};
diff --git a/Engine/Source/Runtime/WebBrowser/Public/SWebBrowser.h b/Engine/Source/Runtime/WebBrowser/Public/SWebBrowser.h
new file mode 100644
--- /dev/null
+++ b/Engine/Source/Runtime/WebBrowser/Public/SWebBrowser.h
@@ -0,0 +1,100 @@
+// Copyright 1998-2015 Epic Games, Inc. All Rights Reserved.
+
+#pragma once
+
+#include "SlateBasics.h"
+
+class IWebBrowserWindow;
+class FWebBrowserViewport;
+
+class WEBBROWSER_API SWebBrowser : public SCompoundWidget
+{
+public:
+	SLATE_BEGIN_ARGS(SWebBrowser)
+		: _InitialURL(TEXT("www.google.com"))
+		, _ShowControls(true)
+		, _SupportsTransparency(false)
+		, _ViewportSize(FVector2D(320, 240))
+		{}
+
+		/** A reference to the parent window */
+		SLATE_ARGUMENT(TSharedPtr<SWindow>, ParentWindow)
+
+		/** URL that the browser will initially navigate to */
+		SLATE_ARGUMENT(FString, InitialURL)
+
+		/** Whether to show standard controls like Back, Forward, Reload etc. */
+		SLATE_ARGUMENT(bool, ShowControls)
+
+		/** Should this browser window support transparency */
+		SLATE_ARGUMENT(bool, SupportsTransparency)
+
+		/** Desired size of the web browser viewport */
+		SLATE_ATTRIBUTE(FVector2D, ViewportSize);
+	SLATE_END_ARGS()
+
+	/**
+	 * Constructor
+	 */
+	SWebBrowser();
+
+	/**
+	 * Construct the widget.
+	 *
+	 * @param InArgs  Declaration from which to construct the widget.
+	 */
+	void Construct(const FArguments& InArgs);
+
+	/**
+	 * Get the current title of the web page
+	 */
+	FText GetTitleText() const;
+
+private:
+	/**
+	 * Returns true if the browser can navigate backwards.
+	 */
+	bool CanGoBack() const;
+
+	/**
+	 * Navigate backwards.
+	 */
+	FReply OnBackClicked();
+
+	/**
+	 * Returns true if the browser can navigate forwards.
+	 */
+	bool CanGoForward() const;
+
+	/**
+	 * Navigate forwards.
+	 */
+	FReply OnForwardClicked();
+
+	/**
+	 * Get text for reload button depending on status
+	 */
+	FText GetReloadButtonText() const;
+
+	/**
+	 * Reload or stop loading
+	 */
+	FReply OnReloadClicked();
+
+	/**
+	 * Get whether the page viewport should be visible
+	 */
+	EVisibility GetViewportVisibility() const;
+
+	/**
+	 * Get whether loading throbber should be visible
+	 */
+	EVisibility GetLoadingThrobberVisibility() const;
+
+private:
+	/** Interface for dealing with a web browser window */
+	TSharedPtr<IWebBrowserWindow>	BrowserWindow;
+
+	/** Viewport interface for rendering the web page */
+	TSharedPtr<FWebBrowserViewport>	BrowserViewport;
+};
diff --git a/Engine/Source/Runtime/WebBrowser/Public/WebBrowserModule.h b/Engine/Source/Runtime/WebBrowser/Public/WebBrowserModule.h
new file mode 100644
--- /dev/null
+++ b/Engine/Source/Runtime/WebBrowser/Public/WebBrowserModule.h
@@ -0,0 +1,32 @@
+// Copyright 1998-2015 Epic Games, Inc. All Rights Reserved.
+
+#pragma once
+
+#include "ModuleManager.h"
+
+// Forward Declarations
+class IWebBrowserSingleton;
+
+/**
+ * WebBrowserModule interface
+ */
+class IWebBrowserModule : public IModuleInterface
+{
+public:
+	/**
+	 * Get or load the Web Browser Module
+	 * 
+	 * @return The loaded module
+	 */
+	static inline IWebBrowserModule& Get()
+	{
+		return FModuleManager::LoadModuleChecked< IWebBrowserModule >("WebBrowser");
+	}
+
+	/**
+	 * Get the Web Browser Singleton
+	 * 
+	 * @return The Web Browser Singleton
+	 */
+	virtual IWebBrowserSingleton* GetSingleton() = 0;
+};
\ No newline at end of file
diff --git a/Engine/Source/Runtime/WebBrowser/Public/WebBrowserViewport.h b/Engine/Source/Runtime/WebBrowser/Public/WebBrowserViewport.h
new file mode 100644
--- /dev/null
+++ b/Engine/Source/Runtime/WebBrowser/Public/WebBrowserViewport.h
@@ -0,0 +1,58 @@
+// Copyright 1998-2015 Epic Games, Inc. All Rights Reserved.
+
+#pragma once
+
+#include "SlateCore.h"
+
+// Forward Declarations
+class IWebBrowserWindow;
+
+/**
+ * A Slate viewport to display a Web Browser Window
+ */
+class WEBBROWSER_API FWebBrowserViewport : public ISlateViewport
+{
+public:
+	/**
+	 * Default Constructor
+	 * 
+	 * @param InWebBrowserWindow The Web Browser Window this viewport will display
+	 * @param InViewportWidget The Widget displaying this viewport (needed to capture mouse)
+	 */
+	FWebBrowserViewport(TSharedPtr<IWebBrowserWindow> InWebBrowserWindow, TSharedPtr<SWidget> InViewportWidget)
+		: WebBrowserWindow(InWebBrowserWindow)
+		, ViewportWidget(InViewportWidget)
+	{ }
+
+	/**
+	 * Destructor.
+	 */
+	~FWebBrowserViewport( )
+	{
+	}
+
+	// ISlateViewport interface
+	virtual FIntPoint GetSize() const override;
+	virtual FSlateShaderResource* GetViewportRenderTargetTexture() const override;
+	virtual void Tick( const FGeometry& AllottedGeometry, double InCurrentTime, float DeltaTime ) override;
+	virtual bool RequiresVsync() const override;
+	virtual FCursorReply OnCursorQuery( const FGeometry& MyGeometry, const FPointerEvent& CursorEvent ) override;
+	virtual FReply OnMouseButtonDown( const FGeometry& MyGeometry, const FPointerEvent& MouseEvent ) override;
+	virtual FReply OnMouseButtonUp( const FGeometry& MyGeometry, const FPointerEvent& MouseEvent ) override;
+	virtual void OnMouseEnter( const FGeometry& MyGeometry, const FPointerEvent& MouseEvent ) override;
+	virtual void OnMouseLeave( const FPointerEvent& MouseEvent ) override;
+	virtual FReply OnMouseMove( const FGeometry& MyGeometry, const FPointerEvent& MouseEvent ) override;
+	virtual FReply OnMouseWheel( const FGeometry& MyGeometry, const FPointerEvent& MouseEvent ) override;
+	virtual FReply OnMouseButtonDoubleClick( const FGeometry& InMyGeometry, const FPointerEvent& InMouseEvent ) override;
+	virtual FReply OnKeyDown( const FGeometry& MyGeometry, const FKeyEvent& InKeyEvent ) override;
+	virtual FReply OnKeyUp( const FGeometry& MyGeometry, const FKeyEvent& InKeyEvent ) override;
+	virtual FReply OnKeyChar( const FGeometry& MyGeometry, const FCharacterEvent& InCharacterEvent ) override;
+	virtual FReply OnFocusReceived( const FFocusEvent& InFocusEvent ) override;
+	virtual void OnFocusLost( const FFocusEvent& InFocusEvent ) override;
+
+private:
+	/** The web browser this viewport will display */
+	TSharedPtr<IWebBrowserWindow>	WebBrowserWindow;
+	/** The viewport widget using this interface */
+	TWeakPtr<SWidget>				ViewportWidget;
+};
diff --git a/Engine/Source/Runtime/WebBrowser/WebBrowser.Build.cs b/Engine/Source/Runtime/WebBrowser/WebBrowser.Build.cs
new file mode 100644
--- /dev/null
+++ b/Engine/Source/Runtime/WebBrowser/WebBrowser.Build.cs
@@ -0,0 +1,31 @@
+// Copyright 1998-2015 Epic Games, Inc. All Rights Reserved.
+
+using UnrealBuildTool;
+
+public class WebBrowser : ModuleRules
+{
+	public WebBrowser(TargetInfo Target)
+	{
+		PublicIncludePaths.Add("Developer/WebBrowser/Public");
+		PrivateIncludePaths.Add("Developer/WebBrowser/Private");
+
+		PrivateDependencyModuleNames.AddRange(
+			new string[]
+			{
+				"Core",
+				"InputCore",
+				"SlateCore",
+				"Slate",
+				"CEF3Utils",
+			}
+		);
+
+		if (Target.Platform == UnrealTargetPlatform.Win64
+        ||  Target.Platform == UnrealTargetPlatform.Mac)
+		{
+			AddThirdPartyPrivateStaticDependencies(Target,
+				"CEF3"
+				);
+		}
+	}
+}
