Added functionality to change the current URL of a web browser window

diff --git a/Engine/Source/Runtime/WebBrowser/Private/SWebBrowser.cpp b/Engine/Source/Runtime/WebBrowser/Private/SWebBrowser.cpp
--- a/Engine/Source/Runtime/WebBrowser/Private/SWebBrowser.cpp
+++ b/Engine/Source/Runtime/WebBrowser/Private/SWebBrowser.cpp
@@ -105,6 +105,14 @@ void SWebBrowser::Construct(const FArguments& InArgs)
 	}
 }
 
+void SWebBrowser::LoadURL(FString NewURL)
+{
+	if (BrowserWindow.IsValid())
+	{
+		BrowserWindow->LoadURL(NewURL);
+	}
+}
+
 FText SWebBrowser::GetTitleText() const
 {
 	if (BrowserWindow.IsValid())
diff --git a/Engine/Source/Runtime/WebBrowser/Private/WebBrowserWindow.cpp b/Engine/Source/Runtime/WebBrowser/Private/WebBrowserWindow.cpp
--- a/Engine/Source/Runtime/WebBrowser/Private/WebBrowserWindow.cpp
+++ b/Engine/Source/Runtime/WebBrowser/Private/WebBrowserWindow.cpp
@@ -15,7 +15,7 @@ FWebBrowserWindow::FWebBrowserWindow(FIntPoint InViewportSize)
 	TextureData.Reserve(ViewportSize.X * ViewportSize.Y * 4);
 	TextureData.SetNumZeroed(ViewportSize.X * ViewportSize.Y * 4);
 
-	if (FSlateApplication::Get().GetRenderer().IsValid())
+	if (FSlateApplication::IsInitialized() && FSlateApplication::Get().GetRenderer().IsValid())
 	{
 		UpdatableTexture = FSlateApplication::Get().GetRenderer()->CreateUpdatableTexture(ViewportSize.X, ViewportSize.Y);
 	}
@@ -32,6 +32,19 @@ FWebBrowserWindow::~FWebBrowserWindow()
 	UpdatableTexture = nullptr;
 }
 
+void FWebBrowserWindow::LoadURL(FString NewURL)
+{
+	if (IsValid())
+	{
+		CefRefPtr<CefFrame> MainFrame = InternalCefBrowser->GetMainFrame();
+		if (MainFrame.get() != nullptr)
+		{
+			CefString URL = *NewURL;
+			MainFrame->LoadURL(URL);
+		}
+	}
+}
+
 void FWebBrowserWindow::SetViewportSize(FVector2D WindowSize)
 {
 	// Magic number for texture size, can't access GetMax2DTextureDimension easily
diff --git a/Engine/Source/Runtime/WebBrowser/Private/WebBrowserWindow.h b/Engine/Source/Runtime/WebBrowser/Private/WebBrowserWindow.h
--- a/Engine/Source/Runtime/WebBrowser/Private/WebBrowserWindow.h
+++ b/Engine/Source/Runtime/WebBrowser/Private/WebBrowserWindow.h
@@ -40,6 +40,7 @@ public:
 	virtual ~FWebBrowserWindow();
 
 	// IWebBrowserWindow Interface
+	virtual void LoadURL(FString NewURL) override;
 	virtual void SetViewportSize(FVector2D WindowSize) override;
 	virtual FSlateShaderResource* GetTexture() override;
 	virtual bool IsValid() const override;
diff --git a/Engine/Source/Runtime/WebBrowser/Public/IWebBrowserWindow.h b/Engine/Source/Runtime/WebBrowser/Public/IWebBrowserWindow.h
--- a/Engine/Source/Runtime/WebBrowser/Public/IWebBrowserWindow.h
+++ b/Engine/Source/Runtime/WebBrowser/Public/IWebBrowserWindow.h
@@ -17,6 +17,13 @@ DECLARE_MULTICAST_DELEGATE_OneParam( FOnTitleChangedDelegate, FString );
 class IWebBrowserWindow
 {
 public:
+	/**
+	 * Load the specified URL
+	 *
+	 * @param NewURL New URL to load
+	 */
+	virtual void LoadURL(FString NewURL) = 0;
+
 	/**
 	 * Set the desired size of the web browser viewport
 	 * 
diff --git a/Engine/Source/Runtime/WebBrowser/Public/SWebBrowser.h b/Engine/Source/Runtime/WebBrowser/Public/SWebBrowser.h
--- a/Engine/Source/Runtime/WebBrowser/Public/SWebBrowser.h
+++ b/Engine/Source/Runtime/WebBrowser/Public/SWebBrowser.h
@@ -45,6 +45,13 @@ public:
 	 */
 	void Construct(const FArguments& InArgs);
 
+	/**
+	 * Load the specified URL
+	 * 
+	 * @param NewURL New URL to load
+	 */
+	void LoadURL(FString NewURL);
+
 	/**
 	 * Get the current title of the web page
 	 */
