Fixed issue with web browser not responding to keyboard input on Mac

Changed how keyboard codes are passed to CEF on Mac

diff --git a/Engine/Source/Runtime/WebBrowser/Private/WebBrowserWindow.cpp b/Engine/Source/Runtime/WebBrowser/Private/WebBrowserWindow.cpp
--- a/Engine/Source/Runtime/WebBrowser/Private/WebBrowserWindow.cpp
+++ b/Engine/Source/Runtime/WebBrowser/Private/WebBrowserWindow.cpp
@@ -112,7 +112,12 @@ void FWebBrowserWindow::OnKeyDown(const FKeyEvent& InKeyEvent)
 	if (IsValid())
 	{
 		CefKeyEvent KeyEvent;
+#if PLATFORM_MAC
+		KeyEvent.native_key_code = InKeyEvent.GetKeyCode();
+		KeyEvent.character = InKeyEvent.GetCharacter();
+#else
 		KeyEvent.windows_key_code = InKeyEvent.GetKeyCode();
+#endif
 		// TODO: Figure out whether this is a system key if we come across problems
 		/*KeyEvent.is_system_key = message == WM_SYSCHAR ||
 			message == WM_SYSKEYDOWN ||
@@ -130,7 +135,12 @@ void FWebBrowserWindow::OnKeyUp(const FKeyEvent& InKeyEvent)
 	if (IsValid())
 	{
 		CefKeyEvent KeyEvent;
+#if PLATFORM_MAC
+		KeyEvent.native_key_code = InKeyEvent.GetKeyCode();
+		KeyEvent.character = InKeyEvent.GetCharacter();
+#else
 		KeyEvent.windows_key_code = InKeyEvent.GetKeyCode();
+#endif
 		// TODO: Figure out whether this is a system key if we come across problems
 		/*KeyEvent.is_system_key = message == WM_SYSCHAR ||
 			message == WM_SYSKEYDOWN ||
@@ -148,7 +158,11 @@ void FWebBrowserWindow::OnKeyChar(const FCharacterEvent& InCharacterEvent)
 	if (IsValid())
 	{
 		CefKeyEvent KeyEvent;
+#if PLATFORM_MAC
+		KeyEvent.character = InCharacterEvent.GetCharacter();
+#else
 		KeyEvent.windows_key_code = InCharacterEvent.GetCharacter();
+#endif
 		// TODO: Figure out whether this is a system key if we come across problems
 		/*KeyEvent.is_system_key = message == WM_SYSCHAR ||
 			message == WM_SYSKEYDOWN ||
