Added log messages for Web Browser errors

Convert relative paths to full path as it can end up incorrect in cooked builds

diff --git a/Engine/Source/Runtime/WebBrowser/Private/WebBrowserSingleton.cpp b/Engine/Source/Runtime/WebBrowser/Private/WebBrowserSingleton.cpp
--- a/Engine/Source/Runtime/WebBrowser/Private/WebBrowserSingleton.cpp
+++ b/Engine/Source/Runtime/WebBrowser/Private/WebBrowserSingleton.cpp
@@ -51,8 +51,18 @@ FWebBrowserSingleton::FWebBrowserSingleton()
 	FString LocalesPath(FPaths::Combine(*FPaths::EngineDir(), TEXT("Binaries/ThirdParty/CEF3/Linux/Resources/locales")));
 #endif
 #if !PLATFORM_MAC // On Mac Chromium ignores custom locales dir. Files need to be stored in Resources folder in the app bundle
+	LocalesPath = FPaths::ConvertRelativePathToFull(LocalesPath);
+	if (!FPaths::DirectoryExists(LocalesPath))
+	{
+		UE_LOG(LogWebBrowser, Error, TEXT("Chromium Locales information not found at: %s."), *LocalesPath);
+	}
 	CefString(&Settings.locales_dir_path) = *LocalesPath;
 #endif
+	ResourcesPath = FPaths::ConvertRelativePathToFull(ResourcesPath);
+	if (!FPaths::DirectoryExists(ResourcesPath))
+	{
+		UE_LOG(LogWebBrowser, Error, TEXT("Chromium Resources information not found at: %s."), *ResourcesPath);
+	}
 	CefString(&Settings.resources_dir_path) = *ResourcesPath;
 
 	// Specify path to sub process exe
@@ -63,6 +73,11 @@ FWebBrowserSingleton::FWebBrowserSingleton()
 #else // @todo Linux
 	FString SubProcessPath(TEXT("UnrealCEFSubProcess"));
 #endif
+	SubProcessPath = FPaths::ConvertRelativePathToFull(SubProcessPath);
+	if (!FPaths::FileExists(SubProcessPath))
+	{
+		UE_LOG(LogWebBrowser, Error, TEXT("UnrealCEFSubProcess.exe not found, check that this program has been built and is placed in: %s."), *SubProcessPath);
+	}
 	CefString(&Settings.browser_subprocess_path) = *SubProcessPath;
 
 	// Initialize CEF.
