Ensured that web browser window size is set correctly

Changed Web Browser Window Set Size function to take an int point as there's no need for float version

Calculated integer size of viewport by adding position and size before converting to int, which avoids problems with rounding.

diff --git a/Engine/Source/Runtime/WebBrowser/Private/WebBrowserViewport.cpp b/Engine/Source/Runtime/WebBrowser/Private/WebBrowserViewport.cpp
--- a/Engine/Source/Runtime/WebBrowser/Private/WebBrowserViewport.cpp
+++ b/Engine/Source/Runtime/WebBrowser/Private/WebBrowserViewport.cpp
@@ -18,7 +18,10 @@ FSlateShaderResource* FWebBrowserViewport::GetViewportRenderTargetTexture() cons
 
 void FWebBrowserViewport::Tick( const FGeometry& AllottedGeometry, double InCurrentTime, float DeltaTime )
 {
-	WebBrowserWindow->SetViewportSize(AllottedGeometry.GetLocalSize());
+	// Calculate max corner of the viewport using same method as Slate
+	FVector2D MaxPos = AllottedGeometry.AbsolutePosition + TransformVector(AllottedGeometry.GetAccumulatedRenderTransform(), AllottedGeometry.GetLocalSize());
+	// Get size by subtracting as int to avoid incorrect rounding when size and position are .5
+	WebBrowserWindow->SetViewportSize(MaxPos.IntPoint() - AllottedGeometry.AbsolutePosition.IntPoint());
 }
 
 bool FWebBrowserViewport::RequiresVsync() const
diff --git a/Engine/Source/Runtime/WebBrowser/Private/WebBrowserWindow.cpp b/Engine/Source/Runtime/WebBrowser/Private/WebBrowserWindow.cpp
--- a/Engine/Source/Runtime/WebBrowser/Private/WebBrowserWindow.cpp
+++ b/Engine/Source/Runtime/WebBrowser/Private/WebBrowserWindow.cpp
@@ -4,6 +4,7 @@
 #include "WebBrowserWindow.h"
 #include "SlateCore.h"
 #include "SlateBasics.h"
+#include "RHI.h"
 
 #if WITH_CEF3
 FWebBrowserWindow::FWebBrowserWindow(FIntPoint InViewportSize)
@@ -45,17 +46,18 @@ void FWebBrowserWindow::LoadURL(FString NewURL)
 	}
 }
 
-void FWebBrowserWindow::SetViewportSize(FVector2D WindowSize)
+void FWebBrowserWindow::SetViewportSize(FIntPoint WindowSize)
 {
-	// Magic number for texture size, can't access GetMax2DTextureDimension easily
-	FIntPoint ClampedWindowSize = WindowSize.ClampAxes(1, 2048).IntPoint();
+	const int32 MaxSize = GetMax2DTextureDimension();
+	WindowSize.X = FMath::Min(WindowSize.X, MaxSize);
+	WindowSize.Y = FMath::Min(WindowSize.Y, MaxSize);
 
 	// Ignore sizes that can't be seen as it forces CEF to re-render whole image
-	if (WindowSize >= FVector2D::UnitVector && ViewportSize != ClampedWindowSize)
+	if (WindowSize.X > 0 && WindowSize.Y > 0 && ViewportSize != WindowSize)
 	{
 		FIntPoint OldViewportSize = MoveTemp(ViewportSize);
 		TArray<uint8> OldTextureData = MoveTemp(TextureData);
-		ViewportSize = ClampedWindowSize;
+		ViewportSize = MoveTemp(WindowSize);
 		TextureData.SetNumZeroed(ViewportSize.X * ViewportSize.Y * 4);
 
 		// copy row by row to avoid texture distortion
diff --git a/Engine/Source/Runtime/WebBrowser/Private/WebBrowserWindow.h b/Engine/Source/Runtime/WebBrowser/Private/WebBrowserWindow.h
--- a/Engine/Source/Runtime/WebBrowser/Private/WebBrowserWindow.h
+++ b/Engine/Source/Runtime/WebBrowser/Private/WebBrowserWindow.h
@@ -41,7 +41,7 @@ public:
 
 	// IWebBrowserWindow Interface
 	virtual void LoadURL(FString NewURL) override;
-	virtual void SetViewportSize(FVector2D WindowSize) override;
+	virtual void SetViewportSize(FIntPoint WindowSize) override;
 	virtual FSlateShaderResource* GetTexture() override;
 	virtual bool IsValid() const override;
 	virtual bool HasBeenPainted() const override;
diff --git a/Engine/Source/Runtime/WebBrowser/Public/IWebBrowserWindow.h b/Engine/Source/Runtime/WebBrowser/Public/IWebBrowserWindow.h
--- a/Engine/Source/Runtime/WebBrowser/Public/IWebBrowserWindow.h
+++ b/Engine/Source/Runtime/WebBrowser/Public/IWebBrowserWindow.h
@@ -29,7 +29,7 @@ public:
 	 * 
 	 * @param WindowSize Desired viewport size
 	 */
-	virtual void SetViewportSize(FVector2D WindowSize) = 0;
+	virtual void SetViewportSize(FIntPoint WindowSize) = 0;
 
 	/**
 	 * Gets interface to the texture representation of the browser
diff --git a/Engine/Source/Runtime/WebBrowser/WebBrowser.Build.cs b/Engine/Source/Runtime/WebBrowser/WebBrowser.Build.cs
--- a/Engine/Source/Runtime/WebBrowser/WebBrowser.Build.cs
+++ b/Engine/Source/Runtime/WebBrowser/WebBrowser.Build.cs
@@ -13,6 +13,7 @@ public class WebBrowser : ModuleRules
 			new string[]
 			{
 				"Core",
+				"RHI",
 				"InputCore",
 				"SlateCore",
 				"Slate",
