Changed Web Browser Viewport to use local size so that image is zoomed when application scale is changed.

This also ensures that the mouse positions line up with what's rendered. Could have also been done by changing browser zoom level so that it stays sharp but it would have meant converting from a scale of 1 to a zoom level of 0 and making all mouse events use absolute positions.

diff --git a/Engine/Source/Runtime/WebBrowser/Private/WebBrowserViewport.cpp b/Engine/Source/Runtime/WebBrowser/Private/WebBrowserViewport.cpp
--- a/Engine/Source/Runtime/WebBrowser/Private/WebBrowserViewport.cpp
+++ b/Engine/Source/Runtime/WebBrowser/Private/WebBrowserViewport.cpp
@@ -19,7 +19,7 @@ FSlateShaderResource* FWebBrowserViewport::GetViewportRenderTargetTexture() cons
 void FWebBrowserViewport::Tick( const FGeometry& AllottedGeometry, double InCurrentTime, float DeltaTime )
 {
 	// Calculate max corner of the viewport using same method as Slate
-	FVector2D MaxPos = AllottedGeometry.AbsolutePosition + TransformVector(AllottedGeometry.GetAccumulatedRenderTransform(), AllottedGeometry.GetLocalSize());
+	FVector2D MaxPos = AllottedGeometry.AbsolutePosition + AllottedGeometry.GetLocalSize();
 	// Get size by subtracting as int to avoid incorrect rounding when size and position are .5
 	WebBrowserWindow->SetViewportSize(MaxPos.IntPoint() - AllottedGeometry.AbsolutePosition.IntPoint());
 }
