Added use of LoadString to Web Browser window

Also added the option to pass a content string when the window is created, as you need to wait a frame before the browser is passed back for it to be valid.

diff --git a/Engine/Source/Runtime/WebBrowser/Private/SWebBrowser.cpp b/Engine/Source/Runtime/WebBrowser/Private/SWebBrowser.cpp
--- a/Engine/Source/Runtime/WebBrowser/Private/SWebBrowser.cpp
+++ b/Engine/Source/Runtime/WebBrowser/Private/SWebBrowser.cpp
@@ -27,7 +27,8 @@ void SWebBrowser::Construct(const FArguments& InArgs)
 	                                                                             InArgs._InitialURL,
 	                                                                             InArgs._ViewportSize.Get().X,
 	                                                                             InArgs._ViewportSize.Get().Y,
-	                                                                             InArgs._SupportsTransparency);
+	                                                                             InArgs._SupportsTransparency,
+	                                                                             InArgs._ContentsToLoad);
 
 	TSharedPtr<SViewport> ViewportWidget;
 
@@ -113,6 +114,14 @@ void SWebBrowser::LoadURL(FString NewURL)
 	}
 }
 
+void SWebBrowser::LoadString(FString Contents, FString DummyURL)
+{
+	if (BrowserWindow.IsValid())
+	{
+		BrowserWindow->LoadString(Contents, DummyURL);
+	}
+}
+
 FText SWebBrowser::GetTitleText() const
 {
 	if (BrowserWindow.IsValid())
diff --git a/Engine/Source/Runtime/WebBrowser/Private/WebBrowserSingleton.cpp b/Engine/Source/Runtime/WebBrowser/Private/WebBrowserSingleton.cpp
--- a/Engine/Source/Runtime/WebBrowser/Private/WebBrowserSingleton.cpp
+++ b/Engine/Source/Runtime/WebBrowser/Private/WebBrowserSingleton.cpp
@@ -114,11 +114,11 @@ FWebBrowserSingleton::~FWebBrowserSingleton()
 #endif
 }
 
-TSharedPtr<IWebBrowserWindow> FWebBrowserSingleton::CreateBrowserWindow(void* OSWindowHandle, FString InitialURL, uint32 Width, uint32 Height, bool bUseTransparency)
+TSharedPtr<IWebBrowserWindow> FWebBrowserSingleton::CreateBrowserWindow(void* OSWindowHandle, FString InitialURL, uint32 Width, uint32 Height, bool bUseTransparency, TOptional<FString> ContentsToLoad)
 {
 #if WITH_CEF3
 	// Create new window
-	TSharedPtr<FWebBrowserWindow> NewWindow(new FWebBrowserWindow(FIntPoint(Width, Height)));
+	TSharedPtr<FWebBrowserWindow> NewWindow(new FWebBrowserWindow(FIntPoint(Width, Height), InitialURL, ContentsToLoad));
 
 	// WebBrowserHandler implements browser-level callbacks.
 	CefRefPtr<FWebBrowserHandler> NewHandler(new FWebBrowserHandler);
diff --git a/Engine/Source/Runtime/WebBrowser/Private/WebBrowserSingleton.h b/Engine/Source/Runtime/WebBrowser/Private/WebBrowserSingleton.h
--- a/Engine/Source/Runtime/WebBrowser/Private/WebBrowserSingleton.h
+++ b/Engine/Source/Runtime/WebBrowser/Private/WebBrowserSingleton.h
@@ -35,7 +35,7 @@ public:
 	virtual ~FWebBrowserSingleton();
 
 	// IWebBrowserSingleton Interface
-	TSharedPtr<IWebBrowserWindow> CreateBrowserWindow(void* OSWindowHandle, FString InitialURL, uint32 Width, uint32 Height, bool bUseTransparency) override;
+	TSharedPtr<IWebBrowserWindow> CreateBrowserWindow(void* OSWindowHandle, FString InitialURL, uint32 Width, uint32 Height, bool bUseTransparency, TOptional<FString> ContentsToLoad = TOptional<FString>()) override;
 
 	// FTickerObjectBase Interface
 	virtual bool Tick(float DeltaTime) override;
diff --git a/Engine/Source/Runtime/WebBrowser/Private/WebBrowserWindow.cpp b/Engine/Source/Runtime/WebBrowser/Private/WebBrowserWindow.cpp
--- a/Engine/Source/Runtime/WebBrowser/Private/WebBrowserWindow.cpp
+++ b/Engine/Source/Runtime/WebBrowser/Private/WebBrowserWindow.cpp
@@ -7,11 +7,13 @@
 #include "RHI.h"
 
 #if WITH_CEF3
-FWebBrowserWindow::FWebBrowserWindow(FIntPoint InViewportSize)
+FWebBrowserWindow::FWebBrowserWindow(FIntPoint InViewportSize, FString InInitialURL, TOptional<FString> InContentsToLoad)
 	: UpdatableTexture(nullptr)
 	, ViewportSize(InViewportSize)
 	, bIsClosing(false)
 	, bHasBeenPainted(false)
+	, InitialURL(InInitialURL)
+	, ContentsToLoad(InContentsToLoad)
 {
 	TextureData.Reserve(ViewportSize.X * ViewportSize.Y * 4);
 	TextureData.SetNumZeroed(ViewportSize.X * ViewportSize.Y * 4);
@@ -46,6 +48,20 @@ void FWebBrowserWindow::LoadURL(FString NewURL)
 	}
 }
 
+void FWebBrowserWindow::LoadString(FString Contents, FString DummyURL)
+{
+	if (IsValid())
+	{
+		CefRefPtr<CefFrame> MainFrame = InternalCefBrowser->GetMainFrame();
+		if (MainFrame.get() != nullptr)
+		{
+			CefString StringVal = *Contents;
+			CefString URL = *DummyURL;
+			MainFrame->LoadString(StringVal, URL);
+		}
+	}
+}
+
 void FWebBrowserWindow::SetViewportSize(FIntPoint WindowSize)
 {
 	const int32 MaxSize = GetMax2DTextureDimension();
@@ -365,6 +381,11 @@ void FWebBrowserWindow::CloseBrowser()
 void FWebBrowserWindow::BindCefBrowser(CefRefPtr<CefBrowser> Browser)
 {
 	InternalCefBrowser = Browser;
+	// Need to wait until this point if we want to start with a page loaded from a string
+	if (ContentsToLoad.IsSet())
+	{
+		LoadString(ContentsToLoad.GetValue(), InitialURL);
+	}
 }
 
 void FWebBrowserWindow::SetTitle(const CefString& InTitle)
diff --git a/Engine/Source/Runtime/WebBrowser/Private/WebBrowserWindow.h b/Engine/Source/Runtime/WebBrowser/Private/WebBrowserWindow.h
--- a/Engine/Source/Runtime/WebBrowser/Private/WebBrowserWindow.h
+++ b/Engine/Source/Runtime/WebBrowser/Private/WebBrowserWindow.h
@@ -32,8 +32,10 @@ public:
 	 * Default Constructor
 	 * 
 	 * @param InViewportSize Initial size of the browser window
+	 * @param InInitialURL The Initial URL that will be loaded
+	 * @param InContentsToLoad Optional string to load as a web page
 	 */
-	FWebBrowserWindow(FIntPoint InViewportSize);
+	FWebBrowserWindow(FIntPoint InViewportSize, FString InInitialURL, TOptional<FString> InContentsToLoad);
 	/**
 	 * Virtual Destructor
 	 */
@@ -41,6 +43,7 @@ public:
 
 	// IWebBrowserWindow Interface
 	virtual void LoadURL(FString NewURL) override;
+	virtual void LoadString(FString Contents, FString DummyURL) override;
 	virtual void SetViewportSize(FIntPoint WindowSize) override;
 	virtual FSlateShaderResource* GetTexture() override;
 	virtual bool IsValid() const override;
@@ -159,6 +162,10 @@ private:
 	bool							bIsClosing;
 	/** Whether this window has been painted at least once */
 	bool							bHasBeenPainted;
+	/** Initial URL that will be loaded (stored to pass on as the dummy for LoadString()) */
+	FString InitialURL;
+	/** Optional text to load as a web page */
+	TOptional<FString> ContentsToLoad;
 
 	// Allow the Handler to access functions only it needs
 	friend class FWebBrowserHandler;
diff --git a/Engine/Source/Runtime/WebBrowser/Public/IWebBrowserSingleton.h b/Engine/Source/Runtime/WebBrowser/Public/IWebBrowserSingleton.h
--- a/Engine/Source/Runtime/WebBrowser/Public/IWebBrowserSingleton.h
+++ b/Engine/Source/Runtime/WebBrowser/Public/IWebBrowserSingleton.h
@@ -25,7 +25,8 @@ public:
 	 * @param Width Initial width of the browser
 	 * @param Height Initial height of the browser
 	 * @param bUseTransparency Whether to allow transparent rendering of pages
+	 * @param ContentsToLoad Optional string to load as a web page
 	 * @return New Web Browser Window Interface (may be null if not supported)
 	 */
-	virtual TSharedPtr<IWebBrowserWindow> CreateBrowserWindow(void* OSWindowHandle, FString InitialURL, uint32 Width, uint32 Height, bool bUseTransparency) = 0;
+	virtual TSharedPtr<IWebBrowserWindow> CreateBrowserWindow(void* OSWindowHandle, FString InitialURL, uint32 Width, uint32 Height, bool bUseTransparency, TOptional<FString> ContentsToLoad = TOptional<FString>()) = 0;
 };
diff --git a/Engine/Source/Runtime/WebBrowser/Public/IWebBrowserWindow.h b/Engine/Source/Runtime/WebBrowser/Public/IWebBrowserWindow.h
--- a/Engine/Source/Runtime/WebBrowser/Public/IWebBrowserWindow.h
+++ b/Engine/Source/Runtime/WebBrowser/Public/IWebBrowserWindow.h
@@ -24,6 +24,14 @@ public:
 	 */
 	virtual void LoadURL(FString NewURL) = 0;
 
+	/**
+	 * Load a string as data to create a web page
+	 *
+	 * @param Contents String to load
+	 * @param DummyURL Dummy URL for the page
+	 */
+	virtual void LoadString(FString Contents, FString DummyURL) = 0;
+
 	/**
 	 * Set the desired size of the web browser viewport
 	 * 
diff --git a/Engine/Source/Runtime/WebBrowser/Public/SWebBrowser.h b/Engine/Source/Runtime/WebBrowser/Public/SWebBrowser.h
--- a/Engine/Source/Runtime/WebBrowser/Public/SWebBrowser.h
+++ b/Engine/Source/Runtime/WebBrowser/Public/SWebBrowser.h
@@ -23,6 +23,9 @@ public:
 		/** URL that the browser will initially navigate to */
 		SLATE_ARGUMENT(FString, InitialURL)
 
+		/** Optional string to load contents as a web page */
+		SLATE_ARGUMENT(TOptional<FString>, ContentsToLoad)
+
 		/** Whether to show standard controls like Back, Forward, Reload etc. */
 		SLATE_ARGUMENT(bool, ShowControls)
 
@@ -52,6 +55,14 @@ public:
 	 */
 	void LoadURL(FString NewURL);
 
+	/**
+	* Load a string as data to create a web page
+	*
+	* @param Contents String to load
+	* @param DummyURL Dummy URL for the page
+	*/
+	void LoadString(FString Contents, FString DummyURL);
+
 	/**
 	 * Get the current title of the web page
 	 */
