WebBrowser: Converted OnTitleChanged delegate to an event; code and documentation cleanup pass
Merging using UE4-To-UE4-LauncherDev, CL# 2441875

diff --git a/Engine/Source/Runtime/WebBrowser/Private/WebBrowserWindow.cpp b/Engine/Source/Runtime/WebBrowser/Private/WebBrowserWindow.cpp
--- a/Engine/Source/Runtime/WebBrowser/Private/WebBrowserWindow.cpp
+++ b/Engine/Source/Runtime/WebBrowser/Private/WebBrowserWindow.cpp
@@ -391,7 +391,7 @@ void FWebBrowserWindow::BindCefBrowser(CefRefPtr<CefBrowser> Browser)
 void FWebBrowserWindow::SetTitle(const CefString& InTitle)
 {
 	Title = InTitle.ToWString().c_str();
-	OnTitleChangedDelegate.Broadcast(Title);
+	TitleChangedEvent.Broadcast(Title);
 }
 
 bool FWebBrowserWindow::GetViewRect(CefRect& Rect)
diff --git a/Engine/Source/Runtime/WebBrowser/Private/WebBrowserWindow.h b/Engine/Source/Runtime/WebBrowser/Private/WebBrowserWindow.h
--- a/Engine/Source/Runtime/WebBrowser/Private/WebBrowserWindow.h
+++ b/Engine/Source/Runtime/WebBrowser/Private/WebBrowserWindow.h
@@ -8,40 +8,61 @@
 #include "WebBrowserHandler.h"
 
 #if PLATFORM_WINDOWS
-#include "AllowWindowsPlatformTypes.h"
+	#include "AllowWindowsPlatformTypes.h"
 #endif
+
 #include "include/internal/cef_ptr.h"
 #include "include/cef_render_handler.h"
+
 #if PLATFORM_WINDOWS
-#include "HideWindowsPlatformTypes.h"
+	#include "HideWindowsPlatformTypes.h"
 #endif
 
-// Forward Declarations
-struct FInputEvent;
 
+struct FInputEvent;
 class FSlateRenderer;
 class FSlateUpdatableTexture;
 
+
 /**
- * Implementation of interface for dealing with a Web Browser window
+ * Implementation of interface for dealing with a Web Browser window.
  */
-class FWebBrowserWindow : public IWebBrowserWindow, public TSharedFromThis<FWebBrowserWindow>
+class FWebBrowserWindow
+	: public IWebBrowserWindow
+	, public TSharedFromThis<FWebBrowserWindow>
 {
+	// Allow the Handler to access functions only it needs
+	friend class FWebBrowserHandler;
+
 public:
 	/**
-	 * Default Constructor
+	 * Creates and initializes a new instance.
 	 * 
-	 * @param InViewportSize Initial size of the browser window
-	 * @param InInitialURL The Initial URL that will be loaded
-	 * @param InContentsToLoad Optional string to load as a web page
+	 * @param InViewportSize Initial size of the browser window.
+	 * @param InInitialURL The Initial URL that will be loaded.
+	 * @param InContentsToLoad Optional string to load as a web page.
 	 */
 	FWebBrowserWindow(FIntPoint InViewportSize, FString InInitialURL, TOptional<FString> InContentsToLoad);
-	/**
-	 * Virtual Destructor
-	 */
+
+	/** Virtual Destructor. */
 	virtual ~FWebBrowserWindow();
 
+public:
+
+	/**
+	* Set the CEF Handler receiving browser callbacks for this window.
+	*
+	* @param InHandler Pointer to the handler for this window.
+	*/
+	void SetHandler(CefRefPtr<FWebBrowserHandler> InHandler);
+
+	/** Close this window so that it can no longer be used. */
+	void CloseBrowser();
+
+public:
+
 	// IWebBrowserWindow Interface
+
 	virtual void LoadURL(FString NewURL) override;
 	virtual void LoadString(FString Contents, FString DummyURL) override;
 	virtual void SetViewportSize(FIntPoint WindowSize) override;
@@ -68,107 +89,114 @@ public:
 	virtual void Reload() override;
 	virtual void StopLoad() override;
 
-	/**
-	 * Set the CEF Handler receiving browser callbacks for this window
-	 *
-	 * @param InHandler Pointer to the handler for this window
-	 */
-	void SetHandler(CefRefPtr<FWebBrowserHandler> InHandler);
-
-	/**
-	 * Close this window so that it can no longer be used
-	 */
-	void CloseBrowser();
+	DECLARE_DERIVED_EVENT(FWebBrowserWindow, IWebBrowserWindow::FOnTitleChanged, FOnTitleChanged);
+	virtual FOnTitleChanged& OnTitleChanged() override
+	{
+		return TitleChangedEvent;
+	}
 
 private:
+
 	/**
 	 * Called to pass reference to the underlying CefBrowser as this is not created at the same time
-	 * as the FWebBrowserWindow
+	 * as the FWebBrowserWindow.
 	 *
-	 * @param Browser The CefBrowser for this window
+	 * @param Browser The CefBrowser for this window.
 	 */
 	void BindCefBrowser(CefRefPtr<CefBrowser> Browser);
 
 	/**
-	 * Sets the Title of this window
+	 * Sets the Title of this window.
 	 * 
-	 * @param InTitle The new title of this window
+	 * @param InTitle The new title of this window.
 	 */
 	void SetTitle(const CefString& InTitle);
 
 	/**
-	 * Get the current proportions of this window
+	 * Get the current proportions of this window.
 	 *
-	 * @param Rect Reference to CefRect to store sizes
-	 * @return Whether Rect was set up correctly
+	 * @param Rect Reference to CefRect to store sizes.
+	 * @return Whether Rect was set up correctly.
 	 */
 	bool GetViewRect(CefRect& Rect);
 
 	/**
-	 * Called when there is an update to the rendered web page
+	 * Called when there is an update to the rendered web page.
 	 *
-	 * @param Type Paint type
-	 * @param DirtyRects List of image areas that have been changed
-	 * @param Buffer Pointer to the raw texture data
-	 * @param Width Width of the texture
-	 * @param Height Height of the texture
+	 * @param Type Paint type.
+	 * @param DirtyRects List of image areas that have been changed.
+	 * @param Buffer Pointer to the raw texture data.
+	 * @param Width Width of the texture.
+	 * @param Height Height of the texture.
 	 */
 	void OnPaint(CefRenderHandler::PaintElementType Type, const CefRenderHandler::RectList& DirtyRects, const void* Buffer, int Width, int Height);
 	
 	/**
-	 * Called when cursor would change due to web browser interaction
+	 * Called when cursor would change due to web browser interaction.
 	 * 
-	 * @param Cursor Handle to CEF mouse cursor
+	 * @param Cursor Handle to CEF mouse cursor.
 	 */
 	void OnCursorChange(CefCursorHandle Cursor);
 
 	/**
-	 * Gets the Cef Keyboard Modifiers based on a Key Event
+	 * Gets the Cef Keyboard Modifiers based on a Key Event.
 	 * 
-	 * @param KeyEvent The Key event
-	 * @return Bits representing keyboard modifiers
+	 * @param KeyEvent The Key event.
+	 * @return Bits representing keyboard modifiers.
 	 */
 	static int32 GetCefKeyboardModifiers(const FKeyEvent& KeyEvent);
 
 	/**
-	 * Gets the Cef Mouse Modifiers based on a Mouse Event
+	 * Gets the Cef Mouse Modifiers based on a Mouse Event.
 	 * 
-	 * @param InMouseEvent The Mouse event
-	 * @return Bits representing mouse modifiers
+	 * @param InMouseEvent The Mouse event.
+	 * @return Bits representing mouse modifiers.
 	 */
 	static int32 GetCefMouseModifiers(const FPointerEvent& InMouseEvent);
 
 	/**
-	 * Gets the Cef Input Modifiers based on an Input Event
+	 * Gets the Cef Input Modifiers based on an Input Event.
 	 * 
-	 * @param InputEvent The Input event
-	 * @return Bits representing input modifiers
+	 * @param InputEvent The Input event.
+	 * @return Bits representing input modifiers.
 	 */
 	static int32 GetCefInputModifiers(const FInputEvent& InputEvent);
 
-	/** Interface to the texture we are rendering to */
-	FSlateUpdatableTexture*			UpdatableTexture;
-	/** Temporary storage for the raw texture data */
-	TArray<uint8>					TextureData;
-	/** Pointer to the CEF Handler for this window */
-	CefRefPtr<FWebBrowserHandler>	Handler;
-	/** Pointer to the CEF Browser for this window */
-	CefRefPtr<CefBrowser>			InternalCefBrowser;
-	/** Current title of this window */
-	FString							Title;
-	/** Current size of this window */
-	FIntPoint						ViewportSize;
-	/** Whether this window is closing */
-	bool							bIsClosing;
-	/** Whether this window has been painted at least once */
-	bool							bHasBeenPainted;
-	/** Initial URL that will be loaded (stored to pass on as the dummy for LoadString()) */
+private:
+
+	/** Interface to the texture we are rendering to. */
+	FSlateUpdatableTexture* UpdatableTexture;
+
+	/** Temporary storage for the raw texture data. */
+	TArray<uint8> TextureData;
+
+	/** Pointer to the CEF Handler for this window. */
+	CefRefPtr<FWebBrowserHandler> Handler;
+
+	/** Pointer to the CEF Browser for this window. */
+	CefRefPtr<CefBrowser> InternalCefBrowser;
+
+	/** Current title of this window. */
+	FString Title;
+
+	/** Current size of this window. */
+	FIntPoint ViewportSize;
+
+	/** Whether this window is closing. */
+	bool bIsClosing;
+
+	/** Whether this window has been painted at least once. */
+	bool bHasBeenPainted;
+
+	/** Initial URL that will be loaded (stored to pass on as the dummy for LoadString()). */
 	FString InitialURL;
-	/** Optional text to load as a web page */
+
+	/** Optional text to load as a web page. */
 	TOptional<FString> ContentsToLoad;
 
-	// Allow the Handler to access functions only it needs
-	friend class FWebBrowserHandler;
+	/** Delegate for broadcasting title changes. */
+	FOnTitleChanged TitleChangedEvent;
 };
 
+
 #endif
diff --git a/Engine/Source/Runtime/WebBrowser/Public/IWebBrowserWindow.h b/Engine/Source/Runtime/WebBrowser/Public/IWebBrowserWindow.h
--- a/Engine/Source/Runtime/WebBrowser/Public/IWebBrowserWindow.h
+++ b/Engine/Source/Runtime/WebBrowser/Public/IWebBrowserWindow.h
@@ -2,14 +2,13 @@
 
 #pragma once
 
-// Forward Declarations
+
 struct FGeometry;
 struct FKeyEvent;
 struct FCharacterEvent;
 struct FPointerEvent;
 class FSlateShaderResource;
 
-DECLARE_MULTICAST_DELEGATE_OneParam( FOnTitleChangedDelegate, FString );
 
 /**
  * Interface for dealing with a Web Browser window
@@ -17,6 +16,7 @@ DECLARE_MULTICAST_DELEGATE_OneParam( FOnTitleChangedDelegate, FString );
 class IWebBrowserWindow
 {
 public:
+
 	/**
 	 * Load the specified URL
 	 *
@@ -132,9 +132,7 @@ public:
 	 */
 	virtual void OnFocus(bool SetFocus) = 0;
 
-	/**
-	 * Called when Capture lost
-	 */
+	/** Called when Capture lost */
 	virtual void OnCaptureLost() = 0;
 
 	/**
@@ -142,9 +140,7 @@ public:
 	 */
 	virtual bool CanGoBack() const = 0;
 
-	/**
-	 * Navigate backwards.
-	 */
+	/** Navigate backwards. */
 	virtual void GoBack() = 0;
 
 	/**
@@ -152,9 +148,7 @@ public:
 	 */
 	virtual bool CanGoForward() const = 0;
 
-	/**
-	 * Navigate forwards.
-	 */
+	/** Navigate forwards. */
 	virtual void GoForward() = 0;
 
 	/**
@@ -162,30 +156,20 @@ public:
 	 */
 	virtual bool IsLoading() const = 0;
 
-	/**
-	 * Reload the current page.
-	 */
+	/** Reload the current page. */
 	virtual void Reload() = 0;
 
-	/**
-	 * Stop loading the page.
-	 */
+	/** Stop loading the page. */
 	virtual void StopLoad() = 0;
 
-	/**
-	 * Accesses a delegate to allow callbacks when a browser title changes
-	 *
-	 * @return A multicast delegate that you can register with
-	 */
-	FOnTitleChangedDelegate& OnTitleChanged() {return OnTitleChangedDelegate;}
+public:
 
-protected:
+	/** A delegate to allow callbacks when a browser title changes. */
+	DECLARE_EVENT_OneParam(IWebBrowserWindow, FOnTitleChanged, FString);
+	virtual FOnTitleChanged& OnTitleChanged() = 0;
 
-	/**
-	 * Virtual Destructor
-	 */
-	virtual ~IWebBrowserWindow() {};
+protected:
 
-	/** Delegate for broadcasting title changes */
-	FOnTitleChangedDelegate OnTitleChangedDelegate;
+	/** Virtual Destructor. */
+	virtual ~IWebBrowserWindow() { };
 };
