[INTEGRATE] UE4-LauncherDev->UE4 integrate all up to cl 2475889

diff --git a/Engine/Source/Runtime/WebBrowser/Public/SWebBrowser.h b/Engine/Source/Runtime/WebBrowser/Public/SWebBrowser.h
--- a/Engine/Source/Runtime/WebBrowser/Public/SWebBrowser.h
+++ b/Engine/Source/Runtime/WebBrowser/Public/SWebBrowser.h
@@ -4,41 +4,60 @@
 
 #include "SlateBasics.h"
 
+
+enum class EWebBrowserDocumentState;
 class IWebBrowserWindow;
 class FWebBrowserViewport;
 
-class WEBBROWSER_API SWebBrowser : public SCompoundWidget
+
+class WEBBROWSER_API SWebBrowser
+	: public SCompoundWidget
 {
 public:
 	SLATE_BEGIN_ARGS(SWebBrowser)
 		: _InitialURL(TEXT("www.google.com"))
 		, _ShowControls(true)
+		, _ShowErrorMessage(true)
 		, _SupportsTransparency(false)
 		, _ViewportSize(FVector2D(320, 240))
-		{}
+	{ }
 
-		/** A reference to the parent window */
+		/** A reference to the parent window. */
 		SLATE_ARGUMENT(TSharedPtr<SWindow>, ParentWindow)
 
-		/** URL that the browser will initially navigate to */
+		/** URL that the browser will initially navigate to. */
 		SLATE_ARGUMENT(FString, InitialURL)
 
-		/** Optional string to load contents as a web page */
+		/** Optional string to load contents as a web page. */
 		SLATE_ARGUMENT(TOptional<FString>, ContentsToLoad)
 
 		/** Whether to show standard controls like Back, Forward, Reload etc. */
 		SLATE_ARGUMENT(bool, ShowControls)
 
-		/** Should this browser window support transparency */
+		/** Whether to show an error message in case of loading errors. */
+		SLATE_ARGUMENT(bool, ShowErrorMessage)
+
+		/** Should this browser window support transparency. */
 		SLATE_ARGUMENT(bool, SupportsTransparency)
 
-		/** Desired size of the web browser viewport */
+		/** Desired size of the web browser viewport. */
 		SLATE_ATTRIBUTE(FVector2D, ViewportSize);
+
+		/** Called when document loading completed. */
+		SLATE_EVENT(FSimpleDelegate, OnLoadCompleted)
+
+		/** Called when document loading failed. */
+		SLATE_EVENT(FSimpleDelegate, OnLoadError)
+
+		/** Called when document loading started. */
+		SLATE_EVENT(FSimpleDelegate, OnLoadStarted)
+
+		/** Called when document title changed. */
+		SLATE_EVENT(FOnTextChanged, OnTitleChanged)
 	SLATE_END_ARGS()
 
-	/**
-	 * Constructor
-	 */
+
+	/** Default constructor. */
 	SWebBrowser();
 
 	/**
@@ -49,70 +68,82 @@ public:
 	void Construct(const FArguments& InArgs);
 
 	/**
-	 * Load the specified URL
+	 * Load the specified URL.
 	 * 
-	 * @param NewURL New URL to load
+	 * @param NewURL New URL to load.
 	 */
 	void LoadURL(FString NewURL);
 
 	/**
-	* Load a string as data to create a web page
+	* Load a string as data to create a web page.
 	*
-	* @param Contents String to load
-	* @param DummyURL Dummy URL for the page
+	* @param Contents String to load.
+	* @param DummyURL Dummy URL for the page.
 	*/
 	void LoadString(FString Contents, FString DummyURL);
 
-	/**
-	 * Get the current title of the web page
-	 */
+	/** Get the current title of the web page. */
 	FText GetTitleText() const;
 
-private:
 	/**
-	 * Returns true if the browser can navigate backwards.
+	 * Gets the currently loaded URL.
+	 *
+	 * @return The URL, or empty string if no document is loaded.
 	 */
+	FString GetUrl();
+
+	/** Whether the document finished loading. */
+	bool IsLoaded() const;
+
+	/** Whether the document is currently being loaded. */
+	bool IsLoading() const; 
+
+private:
+
+	/** Returns true if the browser can navigate backwards. */
 	bool CanGoBack() const;
 
-	/**
-	 * Navigate backwards.
-	 */
+	/** Navigate backwards. */
 	FReply OnBackClicked();
 
-	/**
-	 * Returns true if the browser can navigate forwards.
-	 */
+	/** Returns true if the browser can navigate forwards. */
 	bool CanGoForward() const;
 
-	/**
-	 * Navigate forwards.
-	 */
+	/** Navigate forwards. */
 	FReply OnForwardClicked();
 
-	/**
-	 * Get text for reload button depending on status
-	 */
+	/** Get text for reload button depending on status */
 	FText GetReloadButtonText() const;
 
-	/**
-	 * Reload or stop loading
-	 */
+	/** Reload or stop loading */
 	FReply OnReloadClicked();
 
-	/**
-	 * Get whether the page viewport should be visible
-	 */
+	/** Get whether the page viewport should be visible */
 	EVisibility GetViewportVisibility() const;
 
-	/**
-	 * Get whether loading throbber should be visible
-	 */
+	/** Get whether loading throbber should be visible */
 	EVisibility GetLoadingThrobberVisibility() const;
 
+	/** Callback for document loading state changes. */
+	void HandleBrowserWindowDocumentStateChanged(EWebBrowserDocumentState NewState);
+
 private:
-	/** Interface for dealing with a web browser window */
-	TSharedPtr<IWebBrowserWindow>	BrowserWindow;
 
-	/** Viewport interface for rendering the web page */
-	TSharedPtr<FWebBrowserViewport>	BrowserViewport;
+	/** Interface for dealing with a web browser window. */
+	TSharedPtr<IWebBrowserWindow> BrowserWindow;
+
+	/** Viewport interface for rendering the web page. */
+	TSharedPtr<FWebBrowserViewport> BrowserViewport;
+
+	/** A delegate that is invoked when document loading completed. */
+	FSimpleDelegate OnLoadCompleted;
+
+	/** A delegate that is invoked when document loading failed. */
+	FSimpleDelegate OnLoadError;
+
+	/** A delegate that is invoked when document loading started. */
+	FSimpleDelegate OnLoadStarted;
+
+	/** A delegate that is invoked when document title changed. */
+	FOnTextChanged OnTitleChanged;
 };
