Fixed issues with merge of WebBrowser code from Launcher branch

diff --git a/Engine/Source/Runtime/WebBrowser/Private/WebBrowserWindow.cpp b/Engine/Source/Runtime/WebBrowser/Private/WebBrowserWindow.cpp
--- a/Engine/Source/Runtime/WebBrowser/Private/WebBrowserWindow.cpp
+++ b/Engine/Source/Runtime/WebBrowser/Private/WebBrowserWindow.cpp
@@ -106,7 +106,7 @@ FSlateShaderResource* FWebBrowserWindow::GetTexture()
 	{
         if (bTextureDataDirty)
         {
-            UpdatableTexture->UpdateTexture(TextureData);
+            UpdatableTexture->UpdateTextureThreadSafe(TextureData);
             bTextureDataDirty = false;
         }
 		return UpdatableTexture->GetSlateResource();
@@ -485,11 +485,6 @@ void FWebBrowserWindow::OnPaint(CefRenderHandler::PaintElementType Type, const C
         }
 	}
 
-	if (UpdatableTexture != nullptr)
-	{
-		UpdatableTexture->UpdateTextureThreadSafe(TextureData);
-	}
-
     bTextureDataDirty = true;
     bHasBeenPainted = true;
 
