Adjusted the SWebBrowser's padding for it's internal navigation bar and controls. When these controls are hidden it will no longer pad the viewport 10 pixels from the top.

Merging using UE4-To-UE4-LauncherDev

diff --git a/Engine/Source/Runtime/WebBrowser/Private/SWebBrowser.cpp b/Engine/Source/Runtime/WebBrowser/Private/SWebBrowser.cpp
--- a/Engine/Source/Runtime/WebBrowser/Private/SWebBrowser.cpp
+++ b/Engine/Source/Runtime/WebBrowser/Private/SWebBrowser.cpp
@@ -49,10 +49,11 @@ void SWebBrowser::Construct(const FArguments& InArgs)
 		SNew(SVerticalBox)
 		+SVerticalBox::Slot()
 		.AutoHeight()
-		.Padding(0, 5)
 		[
 			SNew(SHorizontalBox)
-			+SHorizontalBox::Slot()
+			.Visibility((InArgs._ShowControls || InArgs._ShowAddressBar) ? EVisibility::Visible : EVisibility::Collapsed)
+			+ SHorizontalBox::Slot()
+			.Padding(0, 5)
 			.AutoWidth()
 			[
 				SNew(SHorizontalBox)
@@ -95,7 +96,7 @@ void SWebBrowser::Construct(const FArguments& InArgs)
 			+SHorizontalBox::Slot()
 			.VAlign(VAlign_Center)
 			.HAlign(HAlign_Fill)
-			.Padding(5.f, 0.f)
+			.Padding(5.f, 5.f)
 			[
 				// @todo: A proper addressbar widget should go here, for now we use a simple textbox.
 				SAssignNew(InputText, SEditableTextBox)
