[OPP-3051] Remove size clamp imposed on the web browser windows.

Merging CL#2507537 LauncherDev->Main

diff --git a/Engine/Source/Runtime/WebBrowser/Private/WebBrowserWindow.cpp b/Engine/Source/Runtime/WebBrowser/Private/WebBrowserWindow.cpp
--- a/Engine/Source/Runtime/WebBrowser/Private/WebBrowserWindow.cpp
+++ b/Engine/Source/Runtime/WebBrowser/Private/WebBrowserWindow.cpp
@@ -67,10 +67,6 @@ void FWebBrowserWindow::LoadString(FString Contents, FString DummyURL)
 
 void FWebBrowserWindow::SetViewportSize(FIntPoint WindowSize)
 {
-	const int32 MaxSize = GetMax2DTextureDimension();
-	WindowSize.X = FMath::Min(WindowSize.X, MaxSize);
-	WindowSize.Y = FMath::Min(WindowSize.Y, MaxSize);
-
 	// Ignore sizes that can't be seen as it forces CEF to re-render whole image
 	if (WindowSize.X > 0 && WindowSize.Y > 0 && ViewportSize != WindowSize)
 	{
