Integrate UT webbrowser changes

diff --git a/Engine/Source/Runtime/WebBrowser/Private/SWebBrowser.cpp b/Engine/Source/Runtime/WebBrowser/Private/SWebBrowser.cpp
--- a/Engine/Source/Runtime/WebBrowser/Private/SWebBrowser.cpp
+++ b/Engine/Source/Runtime/WebBrowser/Private/SWebBrowser.cpp
@@ -21,6 +21,10 @@ void SWebBrowser::Construct(const FArguments& InArgs)
 	OnLoadStarted = InArgs._OnLoadStarted;
 	OnTitleChanged = InArgs._OnTitleChanged;
 	OnUrlChanged = InArgs._OnUrlChanged;
+	OnJSQueryReceived = InArgs._OnJSQueryReceived;
+	OnJSQueryCanceled = InArgs._OnJSQueryCanceled;
+	OnBeforePopup = InArgs._OnBeforePopup;
+
 	AddressBarUrl = FText::FromString(InArgs._InitialURL);
 	IsHandlingRedraw = false;
 	
@@ -138,6 +142,9 @@ void SWebBrowser::Construct(const FArguments& InArgs)
 		BrowserWindow->OnNeedsRedraw().AddSP(this, &SWebBrowser::HandleBrowserWindowNeedsRedraw);
 		BrowserWindow->OnTitleChanged().AddSP(this, &SWebBrowser::HandleTitleChanged);
 		BrowserWindow->OnUrlChanged().AddSP(this, &SWebBrowser::HandleUrlChanged);
+		BrowserWindow->OnJSQueryReceived().BindSP(this, &SWebBrowser::HandleJSQueryReceived);
+		BrowserWindow->OnJSQueryCanceled().BindSP(this, &SWebBrowser::HandleJSQueryCanceled);
+		BrowserWindow->OnBeforePopup().BindSP(this, &SWebBrowser::HandleBeforePopup);
 		BrowserViewport = MakeShareable(new FWebBrowserViewport(BrowserWindow, ViewportWidget));
 		ViewportWidget->SetViewportInterface(BrowserViewport.ToSharedRef());
 	}
@@ -342,4 +349,37 @@ void SWebBrowser::HandleUrlChanged( FString NewUrl )
 	OnTitleChanged.ExecuteIfBound(AddressBarUrl);
 }
 
+bool SWebBrowser::HandleJSQueryReceived(int64 QueryId, FString QueryString, bool Persistent, FJSQueryResultDelegate Delegate)
+{
+	if (OnJSQueryReceived.IsBound())
+	{
+		return OnJSQueryReceived.Execute(QueryId, QueryString, Persistent, Delegate);
+	}
+	return false;
+}
+
+void SWebBrowser::HandleJSQueryCanceled(int64 QueryId)
+{
+	OnJSQueryCanceled.ExecuteIfBound(QueryId);
+}
+
+bool SWebBrowser::HandleBeforePopup(FString URL, FString Target)
+{
+	if (OnBeforePopup.IsBound())
+	{
+		return OnBeforePopup.Execute(URL, Target);
+	}
+
+	return false;
+}
+
+void SWebBrowser::ExecuteJavascript(const FString& ScriptText)
+{
+	if (BrowserWindow.IsValid())
+	{
+		BrowserWindow->ExecuteJavascript(ScriptText);
+	}
+}
+
 #undef LOCTEXT_NAMESPACE
+
diff --git a/Engine/Source/Runtime/WebBrowser/Private/WebBrowserHandler.cpp b/Engine/Source/Runtime/WebBrowser/Private/WebBrowserHandler.cpp
--- a/Engine/Source/Runtime/WebBrowser/Private/WebBrowserHandler.cpp
+++ b/Engine/Source/Runtime/WebBrowser/Private/WebBrowserHandler.cpp
@@ -180,6 +180,25 @@ bool FWebBrowserHandler::OnBeforeBrowse(CefRefPtr<CefBrowser> Browser,
     return false;
 }
 
+bool FWebBrowserHandler::OnBeforePopup(CefRefPtr<CefBrowser> Browser, 
+	CefRefPtr<CefFrame> Frame, 
+	const CefString& Target_Url, 
+	const CefString& Target_Frame_Name,
+	const CefPopupFeatures& PopupFeatures, 
+	CefWindowInfo& WindowInfo, 
+	CefRefPtr<CefClient>& Client, 
+	CefBrowserSettings& Settings,
+	bool* no_javascript_access)
+{
+	TSharedPtr<FWebBrowserWindow> BrowserWindow = BrowserWindowPtr.Pin();
+	if (BrowserWindow.IsValid())
+	{
+		return BrowserWindow->OnCefBeforePopup(Target_Url, Target_Frame_Name);
+	}
+
+	return false;
+}
+
 void FWebBrowserHandler::SetBrowserWindow(TSharedPtr<FWebBrowserWindow> InBrowserWindow)
 {
 	BrowserWindowPtr = InBrowserWindow;
diff --git a/Engine/Source/Runtime/WebBrowser/Private/WebBrowserHandler.h b/Engine/Source/Runtime/WebBrowser/Private/WebBrowserHandler.h
--- a/Engine/Source/Runtime/WebBrowser/Private/WebBrowserHandler.h
+++ b/Engine/Source/Runtime/WebBrowser/Private/WebBrowserHandler.h
@@ -98,6 +98,15 @@ public:
 
 	virtual void OnAfterCreated(CefRefPtr<CefBrowser> Browser) override;
 	virtual void OnBeforeClose(CefRefPtr<CefBrowser> Browser) override;
+	virtual bool OnBeforePopup(CefRefPtr<CefBrowser> Browser, 
+		CefRefPtr<CefFrame> Frame, 
+		const CefString& Target_Url, 
+		const CefString& Target_Frame_Name,
+		const CefPopupFeatures& PopupFeatures, 
+		CefWindowInfo& WindowInfo,
+		CefRefPtr<CefClient>& Client, 
+		CefBrowserSettings& Settings,
+		bool* no_javascript_access)  override;
 
 public:
 
diff --git a/Engine/Source/Runtime/WebBrowser/Private/WebBrowserWindow.cpp b/Engine/Source/Runtime/WebBrowser/Private/WebBrowserWindow.cpp
--- a/Engine/Source/Runtime/WebBrowser/Private/WebBrowserWindow.cpp
+++ b/Engine/Source/Runtime/WebBrowser/Private/WebBrowserWindow.cpp
@@ -383,6 +383,15 @@ void FWebBrowserWindow::StopLoad()
 	}
 }
 
+void FWebBrowserWindow::ExecuteJavascript(const FString& Script)
+{
+	if (IsValid())
+	{
+		CefRefPtr<CefFrame> frame = InternalCefBrowser->GetMainFrame();
+		frame->ExecuteJavaScript(TCHAR_TO_ANSI(*Script), frame->GetURL(), 0);
+	}
+}
+
 void FWebBrowserWindow::SetHandler(CefRefPtr<FWebBrowserHandler> InHandler)
 {
 	if (InHandler.get())
@@ -583,4 +592,45 @@ int32 FWebBrowserWindow::GetCefInputModifiers(const FInputEvent& InputEvent)
 	return Modifiers;
 }
 
+bool FWebBrowserWindow::OnQuery(int64 QueryId, const CefString& Request, bool Persistent, CefRefPtr<CefMessageRouterBrowserSide::Callback> Callback)
+{
+	if (OnJSQueryReceived().IsBound())
+	{
+		FString QueryString = Request.ToWString().c_str();
+		FJSQueryResultDelegate Delegate = FJSQueryResultDelegate::CreateLambda(
+			[Callback](int ErrorCode, FString Message)
+		{
+			CefString MessageString = *Message;
+			if (ErrorCode == 0)
+			{
+				Callback->Success(MessageString);
+			}
+			else
+			{
+				Callback->Failure(ErrorCode, MessageString);
+			}
+		}
+		);
+		return OnJSQueryReceived().Execute(QueryId, QueryString, Persistent, Delegate);
+	}
+	return false;
+}
+
+void FWebBrowserWindow::OnQueryCanceled(int64 QueryId)
+{
+	OnJSQueryCanceled().ExecuteIfBound(QueryId);
+}
+
+bool FWebBrowserWindow::OnCefBeforePopup(const CefString& Target_Url, const CefString& Target_Frame_Name)
+{
+	if (OnBeforePopup().IsBound())
+	{
+		FString URL = Target_Url.ToWString().c_str();
+		FString FrameName = Target_Frame_Name.ToWString().c_str();
+		return OnBeforePopup().Execute(URL, FrameName);
+	}
+
+	return false;
+}
+
 #endif
diff --git a/Engine/Source/Runtime/WebBrowser/Private/WebBrowserWindow.h b/Engine/Source/Runtime/WebBrowser/Private/WebBrowserWindow.h
--- a/Engine/Source/Runtime/WebBrowser/Private/WebBrowserWindow.h
+++ b/Engine/Source/Runtime/WebBrowser/Private/WebBrowserWindow.h
@@ -91,6 +91,7 @@ public:
 	virtual bool IsLoading() const override;
 	virtual void Reload() override;
 	virtual void StopLoad() override;
+	virtual void ExecuteJavascript(const FString& Script) override;
 
 	DECLARE_DERIVED_EVENT(FWebBrowserWindow, IWebBrowserWindow::FOnDocumentStateChanged, FOnDocumentStateChanged);
 	virtual FOnDocumentStateChanged& OnDocumentStateChanged() override
@@ -116,6 +117,21 @@ public:
 		return NeedsRedrawEvent;
 	}
 
+	virtual FONJSQueryReceived& OnJSQueryReceived() override
+	{
+		return JSQueryReceivedDelegate;
+	}
+
+	virtual FONJSQueryCanceled& OnJSQueryCanceled() override
+	{
+		return JSQueryCanceledDelegate;
+	}
+
+	virtual FOnBeforePopupDelegate& OnBeforePopup() override
+	{
+		return BeforePopupDelegate;
+	}
+
 private:
 
 	/**
@@ -178,6 +194,31 @@ private:
 
 public:
 
+	/**
+	 * Called when JavaScript code sends a message to the UE process.
+	 * Needs to return true or false to tell CEF wether the query is being handled by user code or not.
+	 *
+	 * @param QueryId A unique id for the query. Used to refer to it in OnQueryCanceled.
+	 * @param Request The query string itself as passed in from the JS code.
+	 * @param Persistent Os this a persistent query or not. If not, client code expects the callback to be invoked only once, wheras persistent queries are terminated by invoking Failure, Success can be invoked multiple times until then.
+	 * @param Callback A handle to pass data back to the JS code.
+	 */
+	bool OnQuery(int64 QueryId,
+		const CefString& Request,
+		bool Persistent,
+		CefRefPtr<CefMessageRouterBrowserSide::Callback> Callback);
+
+	/**
+	 * Called when an outstanding query has been canceled either explicitly from JS code or implicitly by navigating away from the page containing the code.
+	 * Will only be called if OnQuery has previously returned true for the same QueryId.
+	 *
+	 * @param QueryId A unique id for the query. A handler should use it to locate and remove any handlers that might be in flight.
+	 */
+	void OnQueryCanceled(int64 QueryId);
+
+	// Trigger an OnBeforePopup event chain
+	bool OnCefBeforePopup(const CefString& Target_Url, const CefString& Target_Frame_Name);
+
 	/**
 	 * Gets the Cef Keyboard Modifiers based on a Key Event.
 	 * 
@@ -254,6 +295,15 @@ private:
 
 	/** Delegate for notifying that the window needs refreshing. */
 	FOnNeedsRedraw NeedsRedrawEvent;
+
+	/** Delegate for notifying that a popup window is attempting to open. */
+	FOnBeforePopupDelegate BeforePopupDelegate;
+
+	/** Delegate for Javascript query received */
+	FONJSQueryReceived JSQueryReceivedDelegate;
+
+	/** Delegate for Javascript query canceled */
+	FONJSQueryCanceled JSQueryCanceledDelegate;
 };
 
 
diff --git a/Engine/Source/Runtime/WebBrowser/Public/IWebBrowserWindow.h b/Engine/Source/Runtime/WebBrowser/Public/IWebBrowserWindow.h
--- a/Engine/Source/Runtime/WebBrowser/Public/IWebBrowserWindow.h
+++ b/Engine/Source/Runtime/WebBrowser/Public/IWebBrowserWindow.h
@@ -181,6 +181,9 @@ public:
 	/** Stop loading the page. */
 	virtual void StopLoad() = 0;
 
+	/** Execute Javascript on the page. */
+	virtual void ExecuteJavascript(const FString& Script) = 0;
+
 public:
 
 	/** A delegate that is invoked when the loading state of a document changed. */
@@ -199,6 +202,20 @@ public:
 	DECLARE_EVENT(IWebBrowserWindow, FOnNeedsRedraw)
 	virtual FOnNeedsRedraw& OnNeedsRedraw() = 0;
 
+	/** A delegate that is invoked when a popup window is attempting to open. */
+	DECLARE_DELEGATE_RetVal_TwoParams(bool, FOnBeforePopupDelegate, FString, FString);
+	virtual FOnBeforePopupDelegate& OnBeforePopup() = 0;
+	
+	/** A delegate that is invoked when JS code sends a query to the front end. The result delegate can either be executed immediately or saved and executed later (multiple times if the boolean persistent argument is true).
+	 * The arguments are an integer error code (0 for success) and a reply string. If you need pass more complex data to the JS code, you will have to pack the data in some way (such as JSON encoding it). */
+	DECLARE_DELEGATE_TwoParams(FJSQueryResultDelegate, int, FString);
+	DECLARE_DELEGATE_RetVal_FourParams(bool, FONJSQueryReceived, int64, FString, bool, FJSQueryResultDelegate)
+	virtual FONJSQueryReceived& OnJSQueryReceived() = 0;
+
+	/** A delegate that is invoked when an outstanding query is canceled. Implement this if you are saving delegates passed to OnQueryReceived. */
+	DECLARE_DELEGATE_OneParam(FONJSQueryCanceled, int64)
+	virtual FONJSQueryCanceled& OnJSQueryCanceled() = 0;
+
 protected:
 
 	/** Virtual Destructor. */
diff --git a/Engine/Source/Runtime/WebBrowser/Public/SWebBrowser.h b/Engine/Source/Runtime/WebBrowser/Public/SWebBrowser.h
--- a/Engine/Source/Runtime/WebBrowser/Public/SWebBrowser.h
+++ b/Engine/Source/Runtime/WebBrowser/Public/SWebBrowser.h
@@ -10,6 +10,11 @@ class IWebBrowserWindow;
 class FWebBrowserViewport;
 
 
+DECLARE_DELEGATE_TwoParams(FJSQueryResultDelegate, int, FString);
+DECLARE_DELEGATE_RetVal_FourParams(bool, FOnJSQueryReceivedDelegate, int64, FString, bool, FJSQueryResultDelegate);
+DECLARE_DELEGATE_OneParam(FOnJSQueryCanceledDelegate, int64);
+DECLARE_DELEGATE_RetVal_TwoParams(bool, FOnBeforePopupDelegate, FString, FString);
+
 class WEBBROWSER_API SWebBrowser
 	: public SCompoundWidget
 {
@@ -60,6 +65,16 @@ public:
 		SLATE_EVENT(FOnTextChanged, OnTitleChanged)
 
 		SLATE_EVENT(FOnTextChanged, OnUrlChanged)
+		
+		/** Called when a custom Javascript message is received from the browser process. */
+		SLATE_EVENT(FOnJSQueryReceivedDelegate, OnJSQueryReceived)
+
+		/** Called when a pending Javascript message has been canceled, either explicitly or by navigating away from the page containing the script. */
+		SLATE_EVENT(FOnJSQueryCanceledDelegate, OnJSQueryCanceled)
+		
+		/** Called before a popup window happens */
+		SLATE_EVENT(FOnBeforePopupDelegate, OnBeforePopup)
+
 	SLATE_END_ARGS()
 
 
@@ -111,6 +126,9 @@ public:
 	/** Whether the document is currently being loaded. */
 	bool IsLoading() const; 
 
+	/** Execute javascript on the current window */
+	void ExecuteJavascript(const FString& ScriptText);
+
 private:
 
 	/** Returns true if the browser can navigate backwards. */
@@ -152,6 +170,15 @@ private:
 	/** Callback for loaded url changes. */
 	void HandleUrlChanged(FString NewUrl);
 
+	/** Callback for received JS queries. */
+	bool HandleJSQueryReceived(int64 QueryId, FString QueryString, bool Persistent, FJSQueryResultDelegate ResultDelegate);
+
+	/** Callback for cancelled JS queries. */
+	void HandleJSQueryCanceled(int64 QueryId);
+
+	/** Callback for popup window permission */
+	bool HandleBeforePopup(FString URL, FString Target);
+
 private:
 
 	/** Interface for dealing with a web browser window. */
@@ -181,6 +208,15 @@ private:
 	/** A delegate that is invoked when document address changed. */
 	FOnTextChanged OnUrlChanged;
 
+	/** A delegate that is invoked when render process Javascript code sends a query message to the client. */
+	FOnJSQueryReceivedDelegate OnJSQueryReceived;
+
+	/** A delegate that is invoked when render process cancels an ongoing query. Handler must clean up corresponding result delegate. */
+	FOnJSQueryCanceledDelegate OnJSQueryCanceled;
+	
+	/** A delegate that is invoked when the browser attempts to pop up a new window */
+	FOnBeforePopupDelegate OnBeforePopup;
+
 	/** A flag to avoid having more than one active timer delegate in flight at the same time */
 	bool IsHandlingRedraw;
 };
