[OPP-3051] Add support for multiple browser windows.

Merging CL#2582871 LauncherDev->Main
change summary:
- Web browser modified to allow users to pass in a callback function that will handle window creation.

diff --git a/Engine/Source/Runtime/WebBrowser/Private/SWebBrowser.cpp b/Engine/Source/Runtime/WebBrowser/Private/SWebBrowser.cpp
--- a/Engine/Source/Runtime/WebBrowser/Private/SWebBrowser.cpp
+++ b/Engine/Source/Runtime/WebBrowser/Private/SWebBrowser.cpp
@@ -14,7 +14,7 @@ SWebBrowser::SWebBrowser()
 {
 }
 
-void SWebBrowser::Construct(const FArguments& InArgs)
+void SWebBrowser::Construct(const FArguments& InArgs, const TSharedPtr<IWebBrowserWindow>& InWebBrowserWindow)
 {
 	OnLoadCompleted = InArgs._OnLoadCompleted;
 	OnLoadError = InArgs._OnLoadError;
@@ -26,27 +26,32 @@ void SWebBrowser::Construct(const FArguments& InArgs)
 	OnJSQueryCanceled = InArgs._OnJSQueryCanceled;
 	OnLoadUrl = InArgs._OnLoadUrl;
 	OnBeforePopup = InArgs._OnBeforePopup;
-
+	OnCreateWindow = InArgs._OnCreateWindow;
 	AddressBarUrl = FText::FromString(InArgs._InitialURL);
 
-	void* OSWindowHandle = nullptr;
-	if (InArgs._ParentWindow.IsValid())
+	BrowserWindow = InWebBrowserWindow;
+	if(!BrowserWindow.IsValid())
 	{
-		TSharedPtr<FGenericWindow> NativeWindow = InArgs._ParentWindow->GetNativeWindow();
-		OSWindowHandle = NativeWindow->GetOSWindowHandle();
+		void* OSWindowHandle = nullptr;
+		if (InArgs._ParentWindow.IsValid())
+		{
+			TSharedPtr<FGenericWindow> NativeWindow = InArgs._ParentWindow->GetNativeWindow();
+			OSWindowHandle = NativeWindow->GetOSWindowHandle();
+		}
+
+		BrowserWindow = IWebBrowserModule::Get().GetSingleton()->CreateBrowserWindow(
+			OSWindowHandle,
+			InArgs._InitialURL,
+			InArgs._ViewportSize.Get().X,
+			InArgs._ViewportSize.Get().Y,
+			InArgs._SupportsTransparency,
+			InArgs._SupportsThumbMouseButtonNavigation,
+			InArgs._ContentsToLoad,
+			InArgs._ShowErrorMessage,
+			InArgs._BackgroundColor
+			);		
 	}
 
-	BrowserWindow = IWebBrowserModule::Get().GetSingleton()->CreateBrowserWindow(
-		OSWindowHandle,
-		InArgs._InitialURL,
-		InArgs._ViewportSize.Get().X,
-		InArgs._ViewportSize.Get().Y,
-		InArgs._SupportsTransparency,
-		InArgs._SupportsThumbMouseButtonNavigation,
-		InArgs._ContentsToLoad,
-		InArgs._ShowErrorMessage,
-		InArgs._BackgroundColor
-	);
 
 	ChildSlot
 	[
@@ -138,6 +143,11 @@ void SWebBrowser::Construct(const FArguments& InArgs)
 
 	if (BrowserWindow.IsValid())
 	{
+		if(OnCreateWindow.IsBound())
+		{
+			BrowserWindow->OnCreateWindow().BindSP(this, &SWebBrowser::HandleCreateWindow);		
+		}
+
 		BrowserWindow->OnDocumentStateChanged().AddSP(this, &SWebBrowser::HandleBrowserWindowDocumentStateChanged);
 		BrowserWindow->OnNeedsRedraw().AddSP(this, &SWebBrowser::HandleBrowserWindowNeedsRedraw);
 		BrowserWindow->OnTitleChanged().AddSP(this, &SWebBrowser::HandleTitleChanged);
@@ -271,7 +281,7 @@ FReply SWebBrowser::OnForwardClicked()
 FText SWebBrowser::GetReloadButtonText() const
 {
 	static FText ReloadText = LOCTEXT("Reload", "Reload");
-	static FText StopText = LOCTEXT("StopText", "StopText");
+	static FText StopText = LOCTEXT("StopText", "Stop");
 
 	if (BrowserWindow.IsValid())
 	{
@@ -430,5 +440,14 @@ void SWebBrowser::ExecuteJavascript(const FString& ScriptText)
 	}
 }
 
+bool SWebBrowser::HandleCreateWindow(const TWeakPtr<IWebBrowserWindow>& NewBrowserWindow)
+{
+	if(OnCreateWindow.IsBound())
+	{
+		return OnCreateWindow.Execute(NewBrowserWindow);
+	}
+	return false;
+}
+
 #undef LOCTEXT_NAMESPACE
 
diff --git a/Engine/Source/Runtime/WebBrowser/Private/WebBrowserHandler.cpp b/Engine/Source/Runtime/WebBrowser/Private/WebBrowserHandler.cpp
--- a/Engine/Source/Runtime/WebBrowser/Private/WebBrowserHandler.cpp
+++ b/Engine/Source/Runtime/WebBrowser/Private/WebBrowserHandler.cpp
@@ -60,8 +60,41 @@ void FWebBrowserHandler::OnAfterCreated(CefRefPtr<CefBrowser> Browser)
 
 	if (BrowserWindow.IsValid())
 	{
+
 		BrowserWindow->BindCefBrowser(Browser);
 	}
+	else
+	{
+		if(Browser->IsPopup())
+		{
+			TSharedPtr<FWebBrowserWindow> BrowserWindowParent = BrowserWindowParentPtr.Pin();
+			if(BrowserWindowParent.IsValid() && BrowserWindowParent->SupportsNewWindows())
+			{
+				TSharedPtr<FWebBrowserWindowInfo> NewBrowserWindowInfo = MakeShareable(new FWebBrowserWindowInfo(Browser, this));
+				TSharedPtr<IWebBrowserWindow> NewBrowserWindow = IWebBrowserModule::Get().GetSingleton()->CreateBrowserWindow( 
+					BrowserWindowParent,
+					NewBrowserWindowInfo
+					);	
+
+				{
+					// @todo: At the moment we need to downcast since the handler does not support using the interface.
+					TSharedPtr<FWebBrowserWindow> HandlerSpecificBrowserWindow = StaticCastSharedPtr<FWebBrowserWindow>(NewBrowserWindow);
+					BrowserWindowPtr = HandlerSpecificBrowserWindow;
+				}
+
+				// Request a UI window for the browser.  If it is not created we do some cleanup.
+				bool bUIWindowCreated = BrowserWindowParent->RequestCreateWindow(NewBrowserWindow.ToSharedRef());
+				if(!bUIWindowCreated)
+				{
+					NewBrowserWindow->CloseBrowser();
+				}
+			}
+			else
+			{
+				Browser->GetHost()->CloseBrowser(true);
+			}
+		}
+	}
 }
 
 void FWebBrowserHandler::OnBeforeClose(CefRefPtr<CefBrowser> Browser)
@@ -85,17 +118,42 @@ bool FWebBrowserHandler::OnBeforePopup( CefRefPtr<CefBrowser> Browser,
 									   CefBrowserSettings& Settings,
 									   bool* NoJavascriptAccess )
 {
+
 	// By default, we ignore any popup requests unless they are handled by us in some way.
 	bool bSupressCEFWindowCreation = true;
-
 	TSharedPtr<FWebBrowserWindow> BrowserWindow = BrowserWindowPtr.Pin();
+
 	if (BrowserWindow.IsValid())
 	{
 		bSupressCEFWindowCreation = BrowserWindow->OnCefBeforePopup(TargetUrl, TArgetFrameName);
+
+		if(!bSupressCEFWindowCreation)
+		{
+			if(BrowserWindow->SupportsNewWindows())
+			{
+				CefRefPtr<FWebBrowserHandler> NewHandler(new FWebBrowserHandler);
+				NewHandler->SetBrowserWindowParent(BrowserWindow);
+				Client = NewHandler;
+				
+				CefWindowHandle ParentWindowHandle = BrowserWindow->GetCefBrowser()->GetHost()->GetWindowHandle();
+
+				// Always use off screen rendering so we can integrate with our windows
+				WindowInfo.SetAsWindowless(ParentWindowHandle, BrowserWindow->UseTransparency());
+
+				// We need to rely on CEF to create our window so we set the WindowInfo, BrowserSettings, Client, and then return false
+				bSupressCEFWindowCreation = false;
+			}
+			else
+			{
+				bSupressCEFWindowCreation = true;
+			}
+		}
 	}
-	return bSupressCEFWindowCreation;
+
+	return bSupressCEFWindowCreation; 
 }
 
+
 void FWebBrowserHandler::OnLoadError(CefRefPtr<CefBrowser> Browser,
 	CefRefPtr<CefFrame> Frame,
 	CefLoadHandler::ErrorCode InErrorCode,
@@ -240,6 +298,11 @@ void FWebBrowserHandler::SetBrowserWindow(TSharedPtr<FWebBrowserWindow> InBrowse
 	BrowserWindowPtr = InBrowserWindow;
 }
 
+void FWebBrowserHandler::SetBrowserWindowParent(TSharedPtr<FWebBrowserWindow> InBrowserWindow)
+{
+	BrowserWindowParentPtr = InBrowserWindow;
+}
+
 bool FWebBrowserHandler::OnProcessMessageReceived(CefRefPtr<CefBrowser> Browser,
 	CefProcessId SourceProcess,
 	CefRefPtr<CefProcessMessage> Message)
diff --git a/Engine/Source/Runtime/WebBrowser/Private/WebBrowserHandler.h b/Engine/Source/Runtime/WebBrowser/Private/WebBrowserHandler.h
--- a/Engine/Source/Runtime/WebBrowser/Private/WebBrowserHandler.h
+++ b/Engine/Source/Runtime/WebBrowser/Private/WebBrowserHandler.h
@@ -48,6 +48,13 @@ public:
 	 * @param InBrowserWindow The browser window this will be handling.
 	 */
 	void SetBrowserWindow(TSharedPtr<FWebBrowserWindow> InBrowserWindow);
+	
+	/**
+	 * Pass in a pointer to our parent Browser Window.
+	 *
+	 * @param InBrowserWindow The parent browser window.
+	 */
+	void SetBrowserWindowParent(TSharedPtr<FWebBrowserWindow> InBrowserWindow);
 
 	/** Sets whether to show messages on loading errors. */
 	void SetShowErrorMessage(bool InShowErrorMessage)
@@ -187,6 +194,9 @@ private:
 
 	/** Weak Pointer to our Web Browser window so that events can be passed on while it's valid*/
 	TWeakPtr<FWebBrowserWindow> BrowserWindowPtr;
+	
+	/** Weak Pointer to the parent web browser window */
+	TWeakPtr<FWebBrowserWindow> BrowserWindowParentPtr;
 
 	/** Whether to show an error message in case of loading errors. */
 	bool ShowErrorMessage;
diff --git a/Engine/Source/Runtime/WebBrowser/Private/WebBrowserSingleton.cpp b/Engine/Source/Runtime/WebBrowser/Private/WebBrowserSingleton.cpp
--- a/Engine/Source/Runtime/WebBrowser/Private/WebBrowserSingleton.cpp
+++ b/Engine/Source/Runtime/WebBrowser/Private/WebBrowserSingleton.cpp
@@ -121,6 +121,33 @@ FWebBrowserSingleton::~FWebBrowserSingleton()
 #endif
 }
 
+TSharedPtr<IWebBrowserWindow> FWebBrowserSingleton::CreateBrowserWindow(
+	TSharedPtr<FWebBrowserWindow>& BrowserWindowParent,
+	TSharedPtr<FWebBrowserWindowInfo>& BrowserWindowInfo
+	)
+{
+#if WITH_CEF3
+
+	//  @todo: Width/Height should be obtained when requesting a UI window from user code which happens later.
+	int32 Width = 800;
+	int32 Height = 600;
+	TOptional<FString> ContentsToLoad;
+
+	bool bShowErrorMessage = BrowserWindowParent->IsShowingErrorMessages();
+	bool bThumbMouseButtonNavigation = BrowserWindowParent->IsThumbMouseButtonNavigationEnabled();
+	bool bUseTransparency = BrowserWindowParent->UseTransparency();
+	FString InitialURL = BrowserWindowInfo->Browser->GetMainFrame()->GetURL().ToWString().c_str();
+	TSharedPtr<FWebBrowserWindow> NewBrowserWindow(new FWebBrowserWindow(FIntPoint(Width, Height), InitialURL, ContentsToLoad, bShowErrorMessage, bThumbMouseButtonNavigation, bUseTransparency));
+	
+	NewBrowserWindow->SetHandler(BrowserWindowInfo->Handler);
+	NewBrowserWindow->BindCefBrowser(BrowserWindowInfo->Browser);
+
+	WindowInterfaces.Add(NewBrowserWindow);
+	return NewBrowserWindow;
+#endif
+	return NULL;
+}
+
 TSharedPtr<IWebBrowserWindow> FWebBrowserSingleton::CreateBrowserWindow(
 	void* OSWindowHandle, 
 	FString InitialURL, 
@@ -134,7 +161,7 @@ TSharedPtr<IWebBrowserWindow> FWebBrowserSingleton::CreateBrowserWindow(
 {
 #if WITH_CEF3
 	// Create new window
-	TSharedPtr<FWebBrowserWindow> NewWindow(new FWebBrowserWindow(FIntPoint(Width, Height), InitialURL, ContentsToLoad, ShowErrorMessage, bThumbMouseButtonNavigation));
+	TSharedPtr<FWebBrowserWindow> NewWindow(new FWebBrowserWindow(FIntPoint(Width, Height), InitialURL, ContentsToLoad, ShowErrorMessage, bThumbMouseButtonNavigation, bUseTransparency));
 
 	// WebBrowserHandler implements browser-level callbacks.
 	CefRefPtr<FWebBrowserHandler> NewHandler(new FWebBrowserHandler);
diff --git a/Engine/Source/Runtime/WebBrowser/Private/WebBrowserSingleton.h b/Engine/Source/Runtime/WebBrowser/Private/WebBrowserSingleton.h
--- a/Engine/Source/Runtime/WebBrowser/Private/WebBrowserSingleton.h
+++ b/Engine/Source/Runtime/WebBrowser/Private/WebBrowserSingleton.h
@@ -49,6 +49,10 @@ public:
 
 	// IWebBrowserSingleton Interface
 
+	TSharedPtr<IWebBrowserWindow> CreateBrowserWindow(
+		TSharedPtr<FWebBrowserWindow>& BrowserWindowParent,
+		TSharedPtr<FWebBrowserWindowInfo>& BrowserWindowInfo) override;
+
 	TSharedPtr<IWebBrowserWindow> CreateBrowserWindow(
 		void* OSWindowHandle, 
 		FString InitialURL, 
@@ -60,6 +64,7 @@ public:
 		bool ShowErrorMessage = true,
 		FColor BackgroundColor = FColor(255, 255, 255, 255)) override;
 
+
 public:
 
 	// FTickerObjectBase Interface
diff --git a/Engine/Source/Runtime/WebBrowser/Private/WebBrowserWindow.cpp b/Engine/Source/Runtime/WebBrowser/Private/WebBrowserWindow.cpp
--- a/Engine/Source/Runtime/WebBrowser/Private/WebBrowserWindow.cpp
+++ b/Engine/Source/Runtime/WebBrowser/Private/WebBrowserWindow.cpp
@@ -13,7 +13,7 @@
 #include <AppKit/NSEvent.h>
 #endif
 
-FWebBrowserWindow::FWebBrowserWindow(FIntPoint InViewportSize, FString InUrl, TOptional<FString> InContentsToLoad, bool InShowErrorMessage, bool InThumbMouseButtonNavigation)
+FWebBrowserWindow::FWebBrowserWindow(FIntPoint InViewportSize, FString InUrl, TOptional<FString> InContentsToLoad, bool InShowErrorMessage, bool InThumbMouseButtonNavigation, bool InUseTransparency)
 	: DocumentState(EWebBrowserDocumentState::NoDocument)
 	, UpdatableTexture(nullptr)
 	, CurrentUrl(InUrl)
@@ -23,6 +23,7 @@ FWebBrowserWindow::FWebBrowserWindow(FIntPoint InViewportSize, FString InUrl, TO
 	, ContentsToLoad(InContentsToLoad)
 	, ShowErrorMessage(InShowErrorMessage)
 	, bThumbMouseButtonNavigation(InThumbMouseButtonNavigation)
+	, bUseTransparency(InUseTransparency)
 	, Cursor(EMouseCursor::Default)
 	, bIsHidden(false)
 	, bTickedLastFrame(true)
@@ -33,10 +34,6 @@ FWebBrowserWindow::FWebBrowserWindow(FIntPoint InViewportSize, FString InUrl, TO
 	, bIgnoreKeyUpEvent(false)
 	, bIgnoreCharacterEvent(false)
 {
-	if (FSlateApplication::IsInitialized() && FSlateApplication::Get().GetRenderer().IsValid())
-	{
-		UpdatableTexture = FSlateApplication::Get().GetRenderer()->CreateUpdatableTexture(1,1); // Texture will be resized on first paint call
-	}
 }
 
 FWebBrowserWindow::~FWebBrowserWindow()
@@ -420,6 +417,21 @@ bool FWebBrowserWindow::OnUnhandledKeyEvent(const CefKeyEvent& CefEvent)
 	return bWasHandled;
 }
 
+
+bool FWebBrowserWindow::SupportsNewWindows()
+{
+	return OnCreateWindow().IsBound() ? true : false;
+}
+
+bool FWebBrowserWindow::RequestCreateWindow( const TSharedRef<IWebBrowserWindow>& NewBrowserWindow )
+{
+	if(OnCreateWindow().IsBound())
+	{
+		return OnCreateWindow().Execute(TWeakPtr<IWebBrowserWindow>(NewBrowserWindow));
+	}
+	return false;
+}
+
 FReply FWebBrowserWindow::OnMouseButtonDown(const FGeometry& MyGeometry, const FPointerEvent& MouseEvent)
 {
 	FReply Reply = FReply::Unhandled();
@@ -668,10 +680,10 @@ void FWebBrowserWindow::CloseBrowser()
 
 void FWebBrowserWindow::BindCefBrowser(CefRefPtr<CefBrowser> Browser)
 {
+	check(Browser.get() == nullptr || InternalCefBrowser.get() == nullptr || InternalCefBrowser->IsSame(Browser));
 	InternalCefBrowser = Browser;
 }
 
-
 CefRefPtr<CefBrowser> FWebBrowserWindow::GetCefBrowser()
 {
 	return InternalCefBrowser;
@@ -743,6 +755,11 @@ void FWebBrowserWindow::NotifyDocumentLoadingStateChange(bool IsLoading)
 
 void FWebBrowserWindow::OnPaint(CefRenderHandler::PaintElementType Type, const CefRenderHandler::RectList& DirtyRects, const void* Buffer, int Width, int Height)
 {
+	if (UpdatableTexture == nullptr && FSlateApplication::IsInitialized() && FSlateApplication::Get().GetRenderer().IsValid())
+	{
+		UpdatableTexture = FSlateApplication::Get().GetRenderer()->CreateUpdatableTexture(Width,Height);
+	}
+
 	if (UpdatableTexture != nullptr)
 	{
 		int32 FirstRow = 0;
diff --git a/Engine/Source/Runtime/WebBrowser/Private/WebBrowserWindow.h b/Engine/Source/Runtime/WebBrowser/Private/WebBrowserWindow.h
--- a/Engine/Source/Runtime/WebBrowser/Private/WebBrowserWindow.h
+++ b/Engine/Source/Runtime/WebBrowser/Private/WebBrowserWindow.h
@@ -23,6 +23,19 @@
 	#include "HideWindowsPlatformTypes.h"
 #endif
 
+/**
+ * Helper for containing items required for CEF browser window creation.
+ */
+struct  FWebBrowserWindowInfo
+{
+	FWebBrowserWindowInfo(CefRefPtr<CefBrowser> InBrowser, CefRefPtr<FWebBrowserHandler> InHandler) 
+		: Browser(InBrowser)
+		, Handler(InHandler)
+	{}
+	CefRefPtr<CefBrowser> Browser;
+	CefRefPtr<FWebBrowserHandler> Handler;
+};
+
 /**
  * Implementation of interface for dealing with a Web Browser window.
  */
@@ -42,7 +55,7 @@ public:
 	 * @param InContentsToLoad Optional string to load as a web page.
 	 * @param InShowErrorMessage Whether to show an error message in case of loading errors.
 	 */
-	FWebBrowserWindow(FIntPoint ViewportSize, FString URL, TOptional<FString> ContentsToLoad, bool ShowErrorMessage, bool bThumbMouseButtonNavigation);
+	FWebBrowserWindow(FIntPoint ViewportSize, FString URL, TOptional<FString> ContentsToLoad, bool ShowErrorMessage, bool bThumbMouseButtonNavigation, bool bUseTransparency);
 
 	/** Virtual Destructor. */
 	virtual ~FWebBrowserWindow();
@@ -56,8 +69,17 @@ public:
 	*/
 	void SetHandler(CefRefPtr<FWebBrowserHandler> InHandler);
 
-	/** Close this window so that it can no longer be used. */
-	void CloseBrowser();
+	/**
+	 * Called to pass reference to the underlying CefBrowser as this is not created at the same time
+	 * as the FWebBrowserWindow.
+	 *
+	 * @param Browser The CefBrowser for this window.
+	 */
+	void BindCefBrowser(CefRefPtr<CefBrowser> Browser);
+
+	bool IsShowingErrorMessages() { return ShowErrorMessage; }
+	bool IsThumbMouseButtonNavigationEnabled() { return bThumbMouseButtonNavigation; }
+	bool UseTransparency() { return bUseTransparency; }
 
 public:
 
@@ -91,6 +113,7 @@ public:
 	virtual void Reload() override;
 	virtual void StopLoad() override;
 	virtual void ExecuteJavascript(const FString& Script) override;
+	virtual void CloseBrowser() override;
 
 	DECLARE_DERIVED_EVENT(FWebBrowserWindow, IWebBrowserWindow::FOnDocumentStateChanged, FOnDocumentStateChanged);
 	virtual FOnDocumentStateChanged& OnDocumentStateChanged() override
@@ -142,6 +165,11 @@ public:
 		return LoadUrlDelegate;
 	}
 
+	virtual FOnCreateWindow& OnCreateWindow() override
+	{
+		return CreateWindowDelegate;
+	}
+
 	virtual FCursorReply OnCursorQuery( const FGeometry& MyGeometry, const FPointerEvent& CursorEvent ) override
 	{
 		return Cursor == EMouseCursor::Default ? FCursorReply::Unhandled() : FCursorReply::Cursor(Cursor);
@@ -154,14 +182,6 @@ public:
 
 private:
 
-	/**
-	 * Called to pass reference to the underlying CefBrowser as this is not created at the same time
-	 * as the FWebBrowserWindow.
-	 *
-	 * @param Browser The CefBrowser for this window.
-	 */
-	void BindCefBrowser(CefRefPtr<CefBrowser> Browser);
-
 	/**
 	 * Used to obtain the internal CEF browser.
 	 *
@@ -269,6 +289,20 @@ private:
 	 */
 	bool OnUnhandledKeyEvent(const CefKeyEvent& CefEvent);
 
+
+	/** Specifies if window creation functionality is available. 
+	 *
+	 * @return true if window creation functionality was provided, false otherwise.  If false, RequestCreateWindow() will always return false.
+	 */
+	bool SupportsNewWindows();
+	
+	/** Called when the browser requests a new UI window
+	 *
+	 * @param NewBrowserWindow The web browser window to display in the new UI window.
+	 * @return true if the UI window was created, false otherwise.
+	 */
+	bool RequestCreateWindow(const TSharedRef<IWebBrowserWindow>& NewBrowserWindow);
+
 public:
 
 	// Trigger an OnBeforePopup event chain
@@ -355,6 +389,9 @@ private:
 	/** Whether to allow forward and back navigation via the mouse thumb buttons. */
 	bool bThumbMouseButtonNavigation;
 
+	/** Whether transparency is enabled. */
+	bool bUseTransparency;
+
 	/** Delegate for broadcasting title changes. */
 	FOnTitleChanged TitleChangedEvent;
 
@@ -378,6 +415,9 @@ private:
 
 	/** Delegate for notifying that a popup window is attempting to open. */
 	FOnBeforePopupDelegate BeforePopupDelegate;
+	
+	/** Delegate for handaling requests to create new windows. */
+	FOnCreateWindow CreateWindowDelegate;
 
 	/** Tracks the current mouse cursor */
 	EMouseCursor::Type Cursor;
diff --git a/Engine/Source/Runtime/WebBrowser/Public/IWebBrowserSingleton.h b/Engine/Source/Runtime/WebBrowser/Public/IWebBrowserSingleton.h
--- a/Engine/Source/Runtime/WebBrowser/Public/IWebBrowserSingleton.h
+++ b/Engine/Source/Runtime/WebBrowser/Public/IWebBrowserSingleton.h
@@ -5,6 +5,8 @@
 // Forward Declarations
 class FSlateRenderer;
 class IWebBrowserWindow;
+class FWebBrowserWindow;
+struct FWebBrowserWindowInfo;
 
 /**
  * A singleton class that takes care of general web browser tasks
@@ -17,6 +19,18 @@ public:
 	 */
 	virtual ~IWebBrowserSingleton() {};
 
+	/**
+	 * Create a new web browser window
+	 *
+	 * @param BrowserWindowParent The parent browser window
+	 * @param BrowserWindowInfo Info for setting up the new browser window
+	 * @return New Web Browser Window Interface (may be null if not supported)
+	 */
+	virtual TSharedPtr<IWebBrowserWindow> CreateBrowserWindow(
+		TSharedPtr<FWebBrowserWindow>& BrowserWindowParent,
+		TSharedPtr<FWebBrowserWindowInfo>& BrowserWindowInfo
+		) = 0;
+
 	/**
 	 * Create a new web browser window
 	 *
diff --git a/Engine/Source/Runtime/WebBrowser/Public/IWebBrowserWindow.h b/Engine/Source/Runtime/WebBrowser/Public/IWebBrowserWindow.h
--- a/Engine/Source/Runtime/WebBrowser/Public/IWebBrowserWindow.h
+++ b/Engine/Source/Runtime/WebBrowser/Public/IWebBrowserWindow.h
@@ -203,6 +203,9 @@ public:
 
 	/** Execute Javascript on the page. */
 	virtual void ExecuteJavascript(const FString& Script) = 0;
+	
+	/** Close this window so that it can no longer be used. */
+	virtual void CloseBrowser() = 0;
 
 public:
 
@@ -247,6 +250,10 @@ public:
 	DECLARE_DELEGATE_RetVal_TwoParams(bool, FOnBeforePopupDelegate, FString, FString);
 	virtual FOnBeforePopupDelegate& OnBeforePopup() = 0;
 
+	/** A delegate this is invoked when an existing browser requests creation of a new browser window. */
+	DECLARE_DELEGATE_RetVal_OneParam(bool, FOnCreateWindow, const TWeakPtr<IWebBrowserWindow>& /*NewBrowserWindow*/)
+	virtual FOnCreateWindow& OnCreateWindow() = 0;
+
 protected:
 
 	/** Virtual Destructor. */
diff --git a/Engine/Source/Runtime/WebBrowser/Public/SWebBrowser.h b/Engine/Source/Runtime/WebBrowser/Public/SWebBrowser.h
--- a/Engine/Source/Runtime/WebBrowser/Public/SWebBrowser.h
+++ b/Engine/Source/Runtime/WebBrowser/Public/SWebBrowser.h
@@ -17,6 +17,7 @@ DECLARE_DELEGATE_TwoParams(FJSQueryResultDelegate, int, FString);
 DECLARE_DELEGATE_RetVal_FourParams(bool, FOnJSQueryReceivedDelegate, int64, FString, bool, FJSQueryResultDelegate);
 DECLARE_DELEGATE_OneParam(FOnJSQueryCanceledDelegate, int64);
 DECLARE_DELEGATE_RetVal_TwoParams(bool, FOnBeforePopupDelegate, FString, FString);
+DECLARE_DELEGATE_RetVal_OneParam(bool, FOnCreateWindowDelegate, const TWeakPtr<IWebBrowserWindow>&);
 
 class WEBBROWSER_API SWebBrowser
 	: public SCompoundWidget
@@ -89,6 +90,9 @@ public:
 		/** Called before a popup window happens */
 		SLATE_EVENT(FOnBeforePopupDelegate, OnBeforePopup)
 
+		/** Called when the browser requests the creation of a new window */
+		SLATE_EVENT(FOnCreateWindowDelegate, OnCreateWindow)
+
 		/** Called before browser navigation. */
 		SLATE_EVENT(FOnBeforeBrowse, OnBeforeNavigation)
 	
@@ -105,7 +109,7 @@ public:
 	 *
 	 * @param InArgs  Declaration from which to construct the widget.
 	 */
-	void Construct(const FArguments& InArgs);
+	void Construct(const FArguments& InArgs, const TSharedPtr<IWebBrowserWindow>& InWebBrowserWindow = nullptr);
 
 	/**
 	 * Load the specified URL.
@@ -217,7 +221,14 @@ private:
 	bool HandleBeforeNavigation(const FString& Url, bool bIsRedirect);
 	
 	bool HandleLoadUrl(const FString& Method, const FString& Url, FString& OutResponse);
-	
+
+	/**
+	 * A delegate that is executed when the browser requests window creation.
+	 *
+	 * @return true if if the window request was handled, false if the browser requesting the new window should be closed.
+	 */
+	bool HandleCreateWindow(const TWeakPtr<IWebBrowserWindow>& NewBrowserWindow);
+
 	/** Callback for popup window permission */
 	bool HandleBeforePopup(FString URL, FString Target);
 
@@ -262,6 +273,9 @@ private:
 	/** A delegate that is invoked when the browser attempts to pop up a new window */
 	FOnBeforePopupDelegate OnBeforePopup;
 
+	/** A delegate that is invoked when the browser requests a UI window for another browser it spawned */
+	FOnCreateWindowDelegate OnCreateWindow;
+
 	/** A delegate that is invoked prior to browser navigation */
 	FOnBeforeBrowse OnBeforeNavigation;
 	
