Merging CL#2601171 using UE4-To-UE4-LauncherDev:
Implement JS integration via a subset of UObject scripting functionality and custom CEF Messages
JIRA: OPP-3240

diff --git a/Engine/Source/Programs/UnrealCEFSubProcess/Private/UnrealCEFSubProcessApp.cpp b/Engine/Source/Programs/UnrealCEFSubProcess/Private/UnrealCEFSubProcessApp.cpp
--- a/Engine/Source/Programs/UnrealCEFSubProcess/Private/UnrealCEFSubProcessApp.cpp
+++ b/Engine/Source/Programs/UnrealCEFSubProcess/Private/UnrealCEFSubProcessApp.cpp
@@ -1,10 +1,4 @@
-//
-//  UnrealCEFSubProcessApp.cpp
-//  UE4
-//
-//  Created by Hrafnkell Freyr Hlodversson on 21/03/15.
-//  Copyright (c) 2015 EpicGames. All rights reserved.
-//
+// Copyright 1998-2015 Epic Games, Inc. All Rights Reserved.
 
 #include "UnrealCEFSubProcess.h"
 #include "UnrealCEFSubProcessApp.h"
@@ -24,13 +18,15 @@ void FUnrealCEFSubProcessApp::OnContextCreated(CefRefPtr<CefBrowser> Browser,
     CefRefPtr<CefFrame> Frame,
     CefRefPtr<CefV8Context> Context)
 {
-    MessageRouter->OnContextCreated(Browser, Frame, Context);
+	RemoteScripting.OnContextCreated(Browser, Frame, Context);
+	MessageRouter->OnContextCreated(Browser, Frame, Context);
 }
 
 void FUnrealCEFSubProcessApp::OnContextReleased(CefRefPtr<CefBrowser> Browser,
     CefRefPtr<CefFrame> Frame,
     CefRefPtr<CefV8Context> Context)
 {
+	RemoteScripting.OnContextReleased(Browser, Frame, Context);
     MessageRouter->OnContextReleased(Browser, Frame, Context);
 }
 
@@ -38,7 +34,35 @@ bool FUnrealCEFSubProcessApp::OnProcessMessageReceived(CefRefPtr<CefBrowser> Bro
     CefProcessId SourceProcess,
     CefRefPtr<CefProcessMessage> Message)
 {
-    return MessageRouter->OnProcessMessageReceived(Browser, SourceProcess, Message);
+	bool Result = false;
+	FString MessageName = Message->GetName().ToWString().c_str();
+	if (MessageName.StartsWith(TEXT("UE::")))
+	{
+		Result = RemoteScripting.OnProcessMessageReceived(Browser, SourceProcess, Message);
+	}
+	else
+	{
+		Result = MessageRouter->OnProcessMessageReceived(Browser, SourceProcess, Message);
+	}
+	return Result;
 }
 
+void FUnrealCEFSubProcessApp::OnRenderThreadCreated(CefRefPtr<CefListValue> ExtraInfo)
+{
+	for(size_t I=0; I<ExtraInfo->GetSize(); I++)
+	{
+		if (ExtraInfo->GetType(I) == VTYPE_DICTIONARY)
+		{
+			CefRefPtr<CefDictionaryValue> Info = ExtraInfo->GetDictionary(I);
+			if ( Info->GetType("browser") == VTYPE_INT)
+			{
+				int32 BrowserID = Info->GetInt("browser");
+				CefRefPtr<CefDictionaryValue> Bindings = Info->GetDictionary("bindings");
+				RemoteScripting.InitPermanentBindings(BrowserID, Bindings);
+			}
+		}
+	}
+}
+
+
 #endif
\ No newline at end of file
diff --git a/Engine/Source/Programs/UnrealCEFSubProcess/Private/UnrealCEFSubProcessApp.h b/Engine/Source/Programs/UnrealCEFSubProcess/Private/UnrealCEFSubProcessApp.h
--- a/Engine/Source/Programs/UnrealCEFSubProcess/Private/UnrealCEFSubProcessApp.h
+++ b/Engine/Source/Programs/UnrealCEFSubProcess/Private/UnrealCEFSubProcessApp.h
@@ -17,6 +17,8 @@
 #include "HideWindowsPlatformTypes.h"
 #endif
 
+#include "UnrealCEFSubProcessRemoteScripting.h"
+
 /**
  * Implements CEF App and other Process level interfaces
  */
@@ -31,7 +33,7 @@ public:
      */
     FUnrealCEFSubProcessApp();
 private:
-    // CefApp methods.
+    // CefApp methods:
     virtual CefRefPtr<CefRenderProcessHandler> GetRenderProcessHandler() override { return this; }
 
     // CefRenderProcessHandler methods:
@@ -47,6 +49,11 @@ private:
         CefProcessId SourceProcess,
         CefRefPtr<CefProcessMessage> Message) override;
 
+	virtual void OnRenderThreadCreated(CefRefPtr<CefListValue> ExtraInfo) override;
+
+	// Handles remote scripting messages from the frontend process
+	FUnrealCEFSubProcessRemoteScripting RemoteScripting;
+
     // Routes messages from Javascript to the browser process and back
     CefRefPtr<CefMessageRouterRendererSide> MessageRouter;
 
diff --git a/Engine/Source/Programs/UnrealCEFSubProcess/Private/UnrealCEFSubProcessCallbackRegistry.cpp b/Engine/Source/Programs/UnrealCEFSubProcess/Private/UnrealCEFSubProcessCallbackRegistry.cpp
new file mode 100644
--- /dev/null
+++ b/Engine/Source/Programs/UnrealCEFSubProcess/Private/UnrealCEFSubProcessCallbackRegistry.cpp
@@ -0,0 +1,39 @@
+// Copyright 1998-2015 Epic Games, Inc. All Rights Reserved.
+
+#include "UnrealCEFSubProcess.h"
+#include "UnrealCEFSubProcessCallbackRegistry.h"
+
+FGuid FUnrealCEFSubProcessCallbackRegistry::FindOrAdd(CefRefPtr<CefV8Context> Context, CefRefPtr<CefV8Value> Object, CefRefPtr<CefV8Value> Function, CefRefPtr<CefV8Value> OnError, bool bOneShot)
+{
+	if (! bOneShot)
+	{
+		for(auto Iter = CreateIterator(); Iter; ++Iter)
+		{
+			auto Value = Iter.Value();
+			if (! Value.bOneShot
+				&& Context->IsSame(Value.Context)
+				&& Function->IsSame(Value.Function)
+				&& Object->IsSame(Value.Object)
+				&& (OnError == nullptr && Value.OnError == nullptr || OnError->IsSame(Value.OnError)))
+			{
+				return Iter.Key();
+			}
+		}
+	}
+
+	// If not found or one-shot, add new entry to the map
+	FGuid Guid = FGuid::NewGuid();
+	Add(Guid, FUnrealCEFSubProcessCallbackRegistryEntry(Context, Object, Function, OnError, bOneShot));
+	return Guid;
+}
+
+void FUnrealCEFSubProcessCallbackRegistry::RemoveByContext(CefRefPtr<CefV8Context> Context)
+{
+	for(auto Iter = CreateIterator(); Iter; ++Iter)
+	{
+		if (Context->IsSame(Iter.Value().Context))
+		{
+			Iter.RemoveCurrent();
+		}
+	}
+}
diff --git a/Engine/Source/Programs/UnrealCEFSubProcess/Private/UnrealCEFSubProcessCallbackRegistry.h b/Engine/Source/Programs/UnrealCEFSubProcess/Private/UnrealCEFSubProcessCallbackRegistry.h
new file mode 100644
--- /dev/null
+++ b/Engine/Source/Programs/UnrealCEFSubProcess/Private/UnrealCEFSubProcessCallbackRegistry.h
@@ -0,0 +1,43 @@
+// Copyright 1998-2015 Epic Games, Inc. All Rights Reserved.
+
+#pragma once
+#include "Guid.h"
+
+#if WITH_CEF3
+
+/** Represents information about a JS function that can be called from the browser process.
+ */
+struct FUnrealCEFSubProcessCallbackRegistryEntry
+{
+	FUnrealCEFSubProcessCallbackRegistryEntry(CefRefPtr<CefV8Context> InContext, CefRefPtr<CefV8Value> InObject, CefRefPtr<CefV8Value> InFunction, CefRefPtr<CefV8Value> InOnError, bool bInOneShot)
+		: Context(InContext)
+		, Object(InObject)
+		, Function(InFunction)
+		, OnError(InOnError)
+		, bOneShot(bInOneShot)
+	{}
+
+
+	CefRefPtr<CefV8Context> Context;
+	CefRefPtr<CefV8Value> Object;
+	CefRefPtr<CefV8Value> Function;
+	CefRefPtr<CefV8Value> OnError;
+	bool bOneShot;
+};
+
+class FUnrealCEFSubProcessCallbackRegistry : public TMap<FGuid, FUnrealCEFSubProcessCallbackRegistryEntry>
+{
+public:
+	/** 
+	 * Looks for a matching entry in the registry or adds a new one, returning its Id.
+	 * A matching entry must have the same Context, Object, Function and OnError object. One-shot callbacks will always result in a new entry , even if there is an exact match.
+	 */
+	FGuid FindOrAdd(CefRefPtr<CefV8Context> Context, CefRefPtr<CefV8Value> Object, CefRefPtr<CefV8Value> Function, CefRefPtr<CefV8Value> OnError=nullptr, bool bOneShot=false);
+
+	/**
+	 * Deletes all entries that have a matching context.
+	 */
+	void RemoveByContext(CefRefPtr<CefV8Context> Context);
+};
+
+#endif
\ No newline at end of file
diff --git a/Engine/Source/Programs/UnrealCEFSubProcess/Private/UnrealCEFSubProcessRemoteMethodHandler.cpp b/Engine/Source/Programs/UnrealCEFSubProcess/Private/UnrealCEFSubProcessRemoteMethodHandler.cpp
new file mode 100644
--- /dev/null
+++ b/Engine/Source/Programs/UnrealCEFSubProcess/Private/UnrealCEFSubProcessRemoteMethodHandler.cpp
@@ -0,0 +1,73 @@
+// Copyright 1998-2015 Epic Games, Inc. All Rights Reserved.
+
+#include "UnrealCEFSubProcess.h"
+#include "UnrealCEFSubProcessRemoteMethodHandler.h"
+#include "UnrealCEFSubProcessRemoteScripting.h"
+
+#if WITH_CEF3
+
+
+bool FUnrealCEFSubProcessRemoteMethodHandler::Execute(const CefString&, CefRefPtr<CefV8Value> Object, const CefV8ValueList& Arguments, CefRefPtr<CefV8Value>& Retval, CefString& InException)
+{
+	return RemoteObject->ExecuteMethod(MethodName, Object, Arguments, Retval, InException);
+}
+
+bool FUnrealCEFSubProcessRemoteObject::ExecuteMethod(const CefString& MethodName, CefRefPtr<CefV8Value> Object, const CefV8ValueList& Arguments, CefRefPtr<CefV8Value>& Retval, CefString& InException)
+{
+	CefRefPtr<CefV8Context> Context = CefV8Context::GetCurrentContext();
+	CefRefPtr<CefV8Exception> Exception;
+	CefRefPtr<CefV8Value> PromiseObjects;
+
+	// Run JS code that creates and unwraps a Promise object
+	if (!Context->Eval(
+		"(function() " \
+		"{ "
+		"	var Accept, Reject, PromiseObject;" \
+		"	PromiseObject = new Promise(function(InAccept, InReject) " \
+		"	{"
+		"		Accept = InAccept;" \
+		"		Reject = InReject;" \
+		"	});" \
+		"	return [PromiseObject, Accept, Reject];" \
+		"})()" , PromiseObjects, Exception))
+	{
+		InException = Exception->GetMessage();
+		return false;
+	}
+
+	if (!(PromiseObjects.get() && PromiseObjects->IsArray() && PromiseObjects->GetArrayLength() == 3))
+	{
+		return false;
+	}
+
+	Retval = PromiseObjects->GetValue(0);
+	CefRefPtr<CefV8Value> Accept = PromiseObjects->GetValue(1);
+	CefRefPtr<CefV8Value> Reject = PromiseObjects->GetValue(2);
+
+	check(Retval->IsObject());
+	check(Accept->IsFunction());
+	check(Reject->IsFunction());
+
+	FGuid CallbackGuid = RemoteScripting->CallbackRegistry.FindOrAdd(Context, Retval, Accept, Reject, true);
+	CefRefPtr<CefProcessMessage> Message = CefProcessMessage::Create("UE::ExecuteUObjectMethod");
+	CefRefPtr<CefListValue> MessageArguments = Message->GetArgumentList();
+	MessageArguments->SetString(0, CefString(*ObjectId.ToString(EGuidFormats::Digits)));
+	MessageArguments->SetString(1, MethodName);
+	MessageArguments->SetString(2, CefString(*CallbackGuid.ToString(EGuidFormats::Digits)));
+	MessageArguments->SetList(3, RemoteScripting->V8ArrayToCef(Arguments));
+
+	Browser->SendProcessMessage(PID_BROWSER, Message);
+
+	return true;
+}
+
+// Destructor will notify the client that the uobject is no longer being referenced on the JS side
+FUnrealCEFSubProcessRemoteObject::~FUnrealCEFSubProcessRemoteObject()
+{
+	CefRefPtr<CefProcessMessage> Message = CefProcessMessage::Create("UE::ReleaseUObject");
+	CefRefPtr<CefListValue> MessageArguments = Message->GetArgumentList();
+	MessageArguments->SetString(0, CefString(*ObjectId.ToString(EGuidFormats::Digits)));
+	Browser->SendProcessMessage(PID_BROWSER, Message);
+}
+
+#endif
\ No newline at end of file
diff --git a/Engine/Source/Programs/UnrealCEFSubProcess/Private/UnrealCEFSubProcessRemoteMethodHandler.h b/Engine/Source/Programs/UnrealCEFSubProcess/Private/UnrealCEFSubProcessRemoteMethodHandler.h
new file mode 100644
--- /dev/null
+++ b/Engine/Source/Programs/UnrealCEFSubProcess/Private/UnrealCEFSubProcessRemoteMethodHandler.h
@@ -0,0 +1,67 @@
+// Copyright 1998-2015 Epic Games, Inc. All Rights Reserved.
+
+#pragma once
+
+#if WITH_CEF3
+#if PLATFORM_WINDOWS
+#include "AllowWindowsPlatformTypes.h"
+#endif
+
+#pragma push_macro("OVERRIDE")
+#undef OVERRIDE // cef headers provide their own OVERRIDE macro
+#include "include/cef_v8.h"
+#pragma pop_macro("OVERRIDE")
+
+#if PLATFORM_WINDOWS
+#include "HideWindowsPlatformTypes.h"
+#endif
+
+class FUnrealCEFSubProcessRemoteScripting;
+
+class FUnrealCEFSubProcessRemoteObject
+	: public CefBase
+{
+public:
+	FUnrealCEFSubProcessRemoteObject(FUnrealCEFSubProcessRemoteScripting* InRemoteScripting, CefRefPtr<CefBrowser> InBrowser, const FGuid& InObjectId)
+		: RemoteScripting(InRemoteScripting)
+		, Browser(InBrowser)
+		, ObjectId(InObjectId) {}
+
+	virtual ~FUnrealCEFSubProcessRemoteObject();
+
+	bool ExecuteMethod(const CefString& MethodName,
+		CefRefPtr<CefV8Value> Object, const CefV8ValueList& Arguments,
+		CefRefPtr<CefV8Value>& Retval, CefString& Exception);
+
+private:
+	FUnrealCEFSubProcessRemoteScripting* RemoteScripting;
+	CefRefPtr<CefBrowser> Browser;
+	FGuid ObjectId;
+
+    // Include the default reference counting implementation.
+    IMPLEMENT_REFCOUNTING(FUnrealCEFSubProcessRemoteObject);
+};
+
+class FUnrealCEFSubProcessRemoteMethodHandler
+	: public CefV8Handler
+{
+
+public:
+	FUnrealCEFSubProcessRemoteMethodHandler(CefRefPtr<FUnrealCEFSubProcessRemoteObject> InRemoteObject, CefString& InMethodName)
+		: RemoteObject(InRemoteObject)
+		, MethodName(InMethodName)
+	{}
+
+	virtual bool Execute(const CefString& Name,
+		CefRefPtr<CefV8Value> Object, const CefV8ValueList& Arguments,
+		CefRefPtr<CefV8Value>& Retval, CefString& Exception) override;
+
+private:
+	CefRefPtr<FUnrealCEFSubProcessRemoteObject> RemoteObject;
+	CefString MethodName;
+
+    // Include the default reference counting implementation.
+    IMPLEMENT_REFCOUNTING(FUnrealCEFSubProcessRemoteMethodHandler);
+};
+
+#endif
\ No newline at end of file
diff --git a/Engine/Source/Programs/UnrealCEFSubProcess/Private/UnrealCEFSubProcessRemoteScripting.cpp b/Engine/Source/Programs/UnrealCEFSubProcess/Private/UnrealCEFSubProcessRemoteScripting.cpp
new file mode 100644
--- /dev/null
+++ b/Engine/Source/Programs/UnrealCEFSubProcess/Private/UnrealCEFSubProcessRemoteScripting.cpp
@@ -0,0 +1,379 @@
+// Copyright 1998-2015 Epic Games, Inc. All Rights Reserved.
+
+#include "UnrealCEFSubProcess.h"
+#include "UnrealCEFSubProcessRemoteScripting.h"
+#include "UnrealCEFSubProcessRemoteMethodHandler.h"
+#if WITH_CEF3
+
+CefRefPtr<CefV8Value> FUnrealCEFSubProcessRemoteScripting::CefToV8(CefRefPtr<CefDictionaryValue> Dictionary)
+{
+	// Custom types are encoded inside dictionary values with a $type and a $value property
+	if ( Dictionary->GetType("$type") == VTYPE_STRING)
+	{
+		FString Type = Dictionary->GetString("$type").ToWString().c_str();
+		if ( Type == "struct" &&  Dictionary->GetType("$value") == VTYPE_DICTIONARY)
+		{
+			return CefToPlainV8Object(Dictionary->GetDictionary("$value"));
+		}
+		if ( Type == "uobject" && Dictionary->GetType("$id") == VTYPE_STRING && Dictionary->GetType("$methods") == VTYPE_LIST)
+		{
+			FGuid Guid;
+			if (FGuid::Parse(Dictionary->GetString("$id").ToWString().c_str(), Guid))
+			{
+				CefRefPtr<CefListValue> Methods = Dictionary->GetList("$methods");
+				return CreateUObjectProxy(Guid, Methods);
+			}
+		}
+	}
+	return CefToPlainV8Object(Dictionary);
+
+}
+
+CefRefPtr<CefV8Value> FUnrealCEFSubProcessRemoteScripting::CreateUObjectProxy(FGuid ObjectId, CefRefPtr<CefListValue> Methods)
+{
+	CefRefPtr<CefV8Context> Context = CefV8Context::GetCurrentContext();
+	CefRefPtr<CefBrowser> Browser = Context->GetBrowser();
+	CefRefPtr<CefV8Value> Result = CefV8Value::CreateObject(nullptr);
+	CefRefPtr<FUnrealCEFSubProcessRemoteObject> Remote = new FUnrealCEFSubProcessRemoteObject(this, Browser, ObjectId);
+	for (size_t I=0; I < Methods->GetSize(); ++I)
+	{
+		CefString MethodName = Methods->GetString(I);
+		CefRefPtr<CefV8Value> FunctionProxy = CefV8Value::CreateFunction(MethodName, new FUnrealCEFSubProcessRemoteMethodHandler(Remote, MethodName));
+		Result->SetValue(MethodName, FunctionProxy, static_cast<CefV8Value::PropertyAttribute>(V8_PROPERTY_ATTRIBUTE_DONTDELETE | V8_PROPERTY_ATTRIBUTE_READONLY ));
+	}
+	Result->SetValue("$id", CefV8Value::CreateString(*ObjectId.ToString(EGuidFormats::Digits)),  static_cast<CefV8Value::PropertyAttribute>(V8_PROPERTY_ATTRIBUTE_DONTDELETE | V8_PROPERTY_ATTRIBUTE_READONLY | V8_PROPERTY_ATTRIBUTE_DONTENUM));
+	return Result;
+}
+
+
+CefRefPtr<CefV8Value> FUnrealCEFSubProcessRemoteScripting::CefToPlainV8Object(CefRefPtr<CefDictionaryValue> Dictionary)
+{
+	CefRefPtr<CefV8Value> Result = CefV8Value::CreateObject(nullptr);
+	CefDictionaryValue::KeyList Keys;
+	Dictionary->GetKeys(Keys);
+	for (CefString Key : Keys)
+	{
+		Result->SetValue(Key, CefToV8(Dictionary, Key), V8_PROPERTY_ATTRIBUTE_NONE);
+	}
+	return Result;
+}
+
+CefRefPtr<CefV8Value> FUnrealCEFSubProcessRemoteScripting::CefToV8(CefRefPtr<CefListValue> List)
+{
+	CefRefPtr<CefV8Value> Result = CefV8Value::CreateArray(List->GetSize());
+	for (size_t i = 0; i < List->GetSize(); ++i)
+	{
+		Result->SetValue(i, CefToV8(List, i));
+	}
+	return Result;
+}
+
+void FUnrealCEFSubProcessRemoteScripting::CefToV8Arglist(CefRefPtr<CefListValue> List, CefV8ValueList& Values)
+{
+	for (size_t i = 0; i < List->GetSize(); ++i)
+	{
+		Values.push_back(CefToV8(List, i));
+	}
+}
+
+
+CefRefPtr<CefListValue> FUnrealCEFSubProcessRemoteScripting::V8ArrayToCef(const CefV8ValueList& Values)
+{
+	CefRefPtr<CefListValue> Result = CefListValue::Create();
+	for (size_t I = 0; I < Values.size(); ++I)
+	{
+		V8ToCef(Result, nullptr, I, Values[I]);
+	}
+	return Result;
+}
+
+CefRefPtr<CefListValue> FUnrealCEFSubProcessRemoteScripting::V8ArrayToCef(CefRefPtr<CefV8Value> Array)
+{
+	CefRefPtr<CefListValue> Result = CefListValue::Create();
+	if (Array->IsArray())
+	{
+		for (int I = 0; I < Array->GetArrayLength(); ++I)
+		{
+			V8ToCef(Result, Array, I, Array->GetValue(I));
+		}
+	}
+	return Result;
+}
+
+CefRefPtr<CefDictionaryValue> FUnrealCEFSubProcessRemoteScripting::V8ObjectToCef(CefRefPtr<CefV8Value> Object)
+{
+	CefRefPtr<CefDictionaryValue> Result = CefDictionaryValue::Create();
+	if (Object->IsObject())
+	{
+		std::vector<CefString> Keys;
+		Object->GetKeys(Keys);
+		for (CefString Key : Keys)
+		{
+			V8ToCef(Result, Object, Key, Object->GetValue(Key));
+		}
+	}
+	return Result;
+}
+
+CefRefPtr<CefDictionaryValue> FUnrealCEFSubProcessRemoteScripting::V8FunctionToCef(CefRefPtr<CefV8Value> Object, CefRefPtr<CefV8Value> Function)
+{
+	CefRefPtr<CefDictionaryValue> Result = CefDictionaryValue::Create();
+	FGuid Guid = CallbackRegistry.FindOrAdd(CefV8Context::GetCurrentContext(), Object, Function);
+	Result->SetString("$type", "callback");
+	Result->SetString("$id", *Guid.ToString(EGuidFormats::Digits));
+	Result->SetString("$name", Function->GetFunctionName());
+	return Result;
+}
+
+bool FUnrealCEFSubProcessRemoteScripting::HandleExecuteJSFunctionMessage(CefRefPtr<CefListValue> MessageArguments)
+{
+	FGuid Guid;
+
+	// Message arguments are CallbackGuid, FunctionArguments, bIsError
+	if (MessageArguments->GetSize() != 3
+		|| MessageArguments->GetType(0) != VTYPE_STRING
+		|| MessageArguments->GetType(1) != VTYPE_LIST
+		|| MessageArguments->GetType(2) != VTYPE_BOOL)
+	{
+		// Wrong message argument types or count
+		return false;
+	}
+
+	if (!FGuid::Parse(FString(MessageArguments->GetString(0).ToWString().c_str()), Guid))
+	{
+		// Invalid GUID
+		return false;
+	}
+	if (!CallbackRegistry.Contains(Guid))
+	{
+		// Unknown callback id
+		return false;
+	}
+
+	auto Callback = CallbackRegistry[Guid];
+	{
+		ScopedV8Context ContextScope(Callback.Context);
+
+		bool bIsErrorCallback = MessageArguments->GetBool(2);
+		CefRefPtr<CefV8Value> Function = bIsErrorCallback?Callback.OnError:Callback.Function;
+
+		if (!Function.get())
+		{
+			// Either invalid entry or no error handler
+			return false;
+		}
+
+		CefV8ValueList FunctionArguments;
+		CefToV8Arglist(MessageArguments->GetList(1), FunctionArguments);
+		CefRefPtr<CefV8Value> Retval = Function->ExecuteFunction(Callback.Object, FunctionArguments);
+		if (!Retval.get())
+		{
+			// Function call resulted in an error
+			return false;
+		}
+
+		// Remove callback if it's a one-shot callback and successful.
+		if (Callback.bOneShot)
+		{
+			CallbackRegistry.Remove(Guid);
+		}
+		return true;
+	}
+}
+
+bool FUnrealCEFSubProcessRemoteScripting::HandleSetValueMessage(CefRefPtr<CefBrowser> Browser, CefRefPtr<CefListValue> MessageArguments)
+{
+	CefRefPtr<CefFrame> MainFrame = Browser->GetMainFrame();
+	CefRefPtr<CefV8Context> Context = MainFrame->GetV8Context();
+	ScopedV8Context ContextScope(Context);
+	CefRefPtr<CefV8Value> RootObject = Context->GetGlobal()->GetValue("ue");
+	if (!RootObject.get())
+	{
+		// The root object should always be created on context creation.
+		return false;
+	}
+
+	for (size_t I = 0; I < MessageArguments->GetSize(); I++)
+	{
+		if (MessageArguments->GetType(I) != VTYPE_DICTIONARY)
+		{
+			return false;
+		}
+		CefRefPtr<CefDictionaryValue> Argument = MessageArguments->GetDictionary(I);
+
+		if (Argument->GetType("name") != VTYPE_STRING
+			|| Argument->GetType("value") != VTYPE_DICTIONARY
+			|| Argument->GetType("permanent") != VTYPE_BOOL)
+		{
+			// Wrong message argument types or count
+			return false;
+		}
+
+		CefString Name = Argument->GetString("name");
+		CefRefPtr<CefDictionaryValue> CefValue = Argument->GetDictionary("value");
+		bool bPermanent = Argument->GetBool("permanent");
+
+		if (bPermanent)
+		{
+			int32 BrowserID = Browser->GetIdentifier();
+			CefRefPtr<CefDictionaryValue> Bindings;
+			if (PermanentBindings.Contains(BrowserID))
+			{
+				Bindings = PermanentBindings[BrowserID];
+			}
+			else
+			{
+				Bindings = CefDictionaryValue::Create();
+				PermanentBindings.Add(BrowserID, Bindings);
+			}
+
+			Bindings->SetDictionary(Name, CefValue);
+		}
+		CefRefPtr<CefV8Value> Value = CefToV8(CefValue);
+		RootObject->SetValue(Name, Value, V8_PROPERTY_ATTRIBUTE_NONE);
+
+	}
+	return true;
+}
+
+bool FUnrealCEFSubProcessRemoteScripting::HandleDeleteValueMessage(CefRefPtr<CefBrowser> Browser, CefRefPtr<CefListValue> MessageArguments)
+{
+	CefRefPtr<CefFrame> MainFrame = Browser->GetMainFrame();
+	CefRefPtr<CefV8Context> Context = MainFrame->GetV8Context();
+	ScopedV8Context ContextScope(Context);
+	CefRefPtr<CefV8Value> RootObject = Context->GetGlobal()->GetValue("ue");
+	if (!RootObject.get())
+	{
+		// The root object should always be created on context creation.
+		return false;
+	}
+
+	for (size_t I = 0; I < MessageArguments->GetSize(); I++)
+	{
+		if (MessageArguments->GetType(I) != VTYPE_DICTIONARY)
+		{
+			return false;
+		}
+		CefRefPtr<CefDictionaryValue> Argument = MessageArguments->GetDictionary(I);
+
+		if (Argument->GetType("name") != VTYPE_STRING
+			|| Argument->GetType("id") != VTYPE_STRING
+			|| Argument->GetType("permanent") != VTYPE_BOOL)
+		{
+			// Wrong message argument types or count
+			return false;
+		}
+
+		CefString Name = Argument->GetString("name");
+		CefString Id = Argument->GetString("id");
+		bool bPermanent = Argument->GetBool("permanent");
+
+		FGuid Guid;
+		if (!FGuid::Parse(Id.ToWString().c_str(), Guid))
+		{
+			return false;
+		}
+
+		if (bPermanent)
+		{
+			int32 BrowserID = Browser->GetIdentifier();
+			CefRefPtr<CefDictionaryValue> Bindings;
+			if (PermanentBindings.Contains(BrowserID))
+			{
+				Bindings = PermanentBindings[BrowserID];
+				if (!Bindings->HasKey(Name))
+				{
+					return false;
+				}
+				if (Guid.IsValid())
+				{
+					CefRefPtr<CefDictionaryValue> CefValue = Bindings->GetDictionary(Name);
+					if (CefValue.get() && CefValue->GetString("$id") != Id)
+					{
+						return false;
+					}
+				}
+				Bindings->Remove(Name);
+			}
+		}
+
+		if (!RootObject->HasValue(Name))
+		{
+			return false;
+		}
+		if (Guid.IsValid())
+		{
+			CefRefPtr<CefV8Value> Value = RootObject->GetValue(Name);
+			if (!Value->HasValue("$id") || !Value->GetValue("$id")->IsString() || Value->GetValue("$id")->GetStringValue() != Id)
+			{
+				return false;
+			}
+		}
+		RootObject->DeleteValue(Name);
+	}
+	return true;
+
+}
+
+
+bool FUnrealCEFSubProcessRemoteScripting::OnProcessMessageReceived(CefRefPtr<CefBrowser> Browser,
+    CefProcessId SourceProcess,
+    CefRefPtr<CefProcessMessage> Message)
+{
+	bool Result = false;
+	FString MessageName = Message->GetName().ToWString().c_str();
+	if (MessageName == TEXT("UE::ExecuteJSFunction"))
+	{
+		Result = HandleExecuteJSFunctionMessage(Message->GetArgumentList());
+	}
+	else if (MessageName == TEXT("UE::SetValue"))
+	{
+		Result = HandleSetValueMessage(Browser, Message->GetArgumentList());
+	}
+	else if (MessageName == TEXT("UE::DeleteValue"))
+	{
+		Result = HandleDeleteValueMessage(Browser, Message->GetArgumentList());
+	}
+
+	return Result;
+}
+
+void FUnrealCEFSubProcessRemoteScripting::OnContextCreated(CefRefPtr<CefBrowser> Browser, CefRefPtr<CefFrame> Frame, CefRefPtr<CefV8Context> Context)
+{
+	ScopedV8Context ContextScope(Context);
+	CefRefPtr<CefV8Value> Global = Context->GetGlobal();
+	if ( !Global->HasValue("ue") )
+	{
+		Global->SetValue("ue", CefV8Value::CreateObject(nullptr), V8_PROPERTY_ATTRIBUTE_DONTDELETE);
+	}
+	CefRefPtr<CefV8Value> RootObject = Context->GetGlobal()->GetValue("ue");
+
+	int32 BrowserID = Browser->GetIdentifier();
+
+	if (PermanentBindings.Contains(BrowserID))
+	{
+		CefRefPtr<CefDictionaryValue> Bindings = PermanentBindings[BrowserID];
+		CefDictionaryValue::KeyList Keys;
+		Bindings->GetKeys(Keys);
+		for (CefString Key : Keys)
+		{
+			CefRefPtr<CefV8Value> Value = CefToV8(Bindings->GetDictionary(Key));
+			RootObject->SetValue(Key, Value, V8_PROPERTY_ATTRIBUTE_NONE);
+		}
+	}
+}
+
+void FUnrealCEFSubProcessRemoteScripting::OnContextReleased(CefRefPtr<CefBrowser> Browser, CefRefPtr<CefFrame> Frame, CefRefPtr<CefV8Context> Context)
+{
+	// Invalidate JS functions that were created in the context being released.
+	CallbackRegistry.RemoveByContext(Context);
+}
+
+void FUnrealCEFSubProcessRemoteScripting::InitPermanentBindings(int32 BrowserID, CefRefPtr<CefDictionaryValue> Values)
+{
+	// The CefDictionary in PermanentBindings needs to be writable, so we'll copy the Values argument before saving it.
+	PermanentBindings.Add(BrowserID, Values->Copy(true));
+}
+
+
+#endif
\ No newline at end of file
diff --git a/Engine/Source/Programs/UnrealCEFSubProcess/Private/UnrealCEFSubProcessRemoteScripting.h b/Engine/Source/Programs/UnrealCEFSubProcess/Private/UnrealCEFSubProcessRemoteScripting.h
new file mode 100644
--- /dev/null
+++ b/Engine/Source/Programs/UnrealCEFSubProcess/Private/UnrealCEFSubProcessRemoteScripting.h
@@ -0,0 +1,160 @@
+// Copyright 1998-2015 Epic Games, Inc. All Rights Reserved.
+
+#pragma once
+
+#if WITH_CEF3
+#if PLATFORM_WINDOWS
+#include "AllowWindowsPlatformTypes.h"
+#endif
+
+#pragma push_macro("OVERRIDE")
+#undef OVERRIDE // cef headers provide their own OVERRIDE macro
+#include "include/cef_app.h"
+#include "include/cef_v8.h"
+#pragma pop_macro("OVERRIDE")
+
+#if PLATFORM_WINDOWS
+#include "HideWindowsPlatformTypes.h"
+#endif
+
+#include "UnrealCEFSubProcessCallbackRegistry.h"
+
+
+class FUnrealCEFSubProcessRemoteScripting
+{
+
+public:
+
+	FUnrealCEFSubProcessRemoteScripting()
+	{}
+
+	CefRefPtr<CefV8Value> CefToV8(CefRefPtr<CefDictionaryValue> Dictionary);
+	CefRefPtr<CefV8Value> CefToV8(CefRefPtr<CefListValue> List);
+	void CefToV8Arglist(CefRefPtr<CefListValue> List, CefV8ValueList& Values);
+	// Works for CefListValue and CefDictionaryValues
+	template<typename ContainerType, typename KeyType>
+	CefRefPtr<CefV8Value> CefToV8(CefRefPtr<ContainerType> Container, KeyType Key)
+	{
+		switch(Container->GetType(Key))
+		{
+			case VTYPE_NULL:
+				return CefV8Value::CreateNull();
+			case VTYPE_BOOL:
+				return CefV8Value::CreateBool(Container->GetBool(Key));
+			case VTYPE_INT:
+				return CefV8Value::CreateInt(Container->GetInt(Key));
+			case VTYPE_DOUBLE:
+				return CefV8Value::CreateDouble(Container->GetDouble(Key));
+			case VTYPE_STRING:
+				return CefV8Value::CreateString(Container->GetString(Key));
+			case VTYPE_DICTIONARY:
+				return CefToV8(Container->GetDictionary(Key));
+			case VTYPE_LIST:
+				return CefToV8(Container->GetList(Key));
+			default:
+				return nullptr;
+		}
+
+	}
+
+	template<typename ContainerType, typename KeyType>
+	bool V8ToCef(CefRefPtr<ContainerType> Container, CefRefPtr<CefV8Value> Parent, KeyType Key, CefRefPtr<CefV8Value> Value)
+	{
+		if(Value->IsNull())
+		{
+			return Container->SetNull(Key);
+		}
+		else if(Value->IsUndefined())
+		{
+			return Container->SetNull(Key); // TODO distinguish between NULL and undef?
+		}
+		else if(Value->IsBool())
+		{
+			return Container->SetBool(Key, Value->GetBoolValue());
+		}
+		else if(Value->IsDouble())
+		{
+			return Container->SetDouble(Key, Value->GetDoubleValue());
+		}
+		else if(Value->IsInt())
+		{
+			return Container->SetInt(Key, Value->GetIntValue());
+		}
+		else if(Value->IsString())
+		{
+			return Container->SetString(Key, Value->GetStringValue());
+		}
+		else if(Value->IsUInt())
+		{
+			return Container->SetDouble(Key, Value->GetUIntValue());
+		}
+		else if(Value->IsDate())
+		{
+			return Container->SetNull(Key); // TODO
+		}
+		else if(Value->IsFunction())
+		{
+			return Container->SetDictionary(Key, V8FunctionToCef(Parent, Value));
+		}
+		else if(Value->IsArray())
+		{
+			return Container->SetList(Key, V8ArrayToCef(Value));
+		}
+		else if(Value->IsObject())
+		{
+			return Container->SetDictionary(Key, V8ObjectToCef(Value));
+		}
+		else
+		{
+			return Container->SetNull(Key);
+		}
+	}
+	CefRefPtr<CefListValue> V8ArrayToCef(const CefV8ValueList& Values);
+	CefRefPtr<CefListValue> V8ArrayToCef(CefRefPtr<CefV8Value> Value);
+	CefRefPtr<CefDictionaryValue> V8ObjectToCef(CefRefPtr<CefV8Value> Value);
+	CefRefPtr<CefDictionaryValue> V8FunctionToCef(CefRefPtr<CefV8Value> Object, CefRefPtr<CefV8Value> Function);
+
+
+    bool OnProcessMessageReceived(CefRefPtr<CefBrowser> Browser, CefProcessId SourceProcess, CefRefPtr<CefProcessMessage> Message);
+	void OnContextCreated(CefRefPtr<CefBrowser> Browser, CefRefPtr<CefFrame> Frame, CefRefPtr<CefV8Context> Context);
+	void OnContextReleased(CefRefPtr<CefBrowser> Browser, CefRefPtr<CefFrame> Frame, CefRefPtr<CefV8Context> Context);
+	void InitPermanentBindings(int32 BrowserID, CefRefPtr<CefDictionaryValue> Values);
+
+private:
+	// Scripting integration message handlers:
+	bool HandleExecuteJSFunctionMessage(CefRefPtr<CefListValue> MessageArguments);
+	bool HandleSetValueMessage(CefRefPtr<CefBrowser> Browser, CefRefPtr<CefListValue> MessageArguments);
+	bool HandleDeleteValueMessage(CefRefPtr<CefBrowser> Browser, CefRefPtr<CefListValue> MessageArguments);
+
+	/** Deletes information about a UObject on the client side.
+	 */
+	bool HandleUnRegisterUObjectMessage(CefRefPtr<CefListValue> MessageArguments);
+
+	CefRefPtr<CefV8Value> CefToPlainV8Object(CefRefPtr<CefDictionaryValue> Dictionary);
+	CefRefPtr<CefV8Value> CreateUObjectProxy(FGuid ObjectId, CefRefPtr<CefListValue> Methods);
+
+	// Stores information about functions that can be called from the remote process.
+	FUnrealCEFSubProcessCallbackRegistry CallbackRegistry;
+
+	TMap<int32,CefRefPtr<CefDictionaryValue>> PermanentBindings;
+
+	class ScopedV8Context
+	{
+	public:
+		ScopedV8Context(CefRefPtr<CefV8Context> InContext)
+			: Context(InContext)
+		{
+			Context->Enter();
+		}
+		~ScopedV8Context()
+		{
+			Context->Exit();
+		}
+	private:
+		CefRefPtr<CefV8Context> Context;
+	};
+
+	friend class FUnrealCEFSubProcessRemoteObject;
+};
+
+#endif
diff --git a/Engine/Source/Runtime/WebBrowser/Private/SWebBrowser.cpp b/Engine/Source/Runtime/WebBrowser/Private/SWebBrowser.cpp
--- a/Engine/Source/Runtime/WebBrowser/Private/SWebBrowser.cpp
+++ b/Engine/Source/Runtime/WebBrowser/Private/SWebBrowser.cpp
@@ -449,5 +449,21 @@ bool SWebBrowser::HandleCreateWindow(const TWeakPtr<IWebBrowserWindow>& NewBrows
 	return false;
 }
 
+void SWebBrowser::BindUObject(const FString& Name, UObject* Object, bool bIsPermanent)
+{
+	if (BrowserWindow.IsValid())
+	{
+		BrowserWindow->BindUObject(Name, Object, bIsPermanent);
+	}
+}
+
+void SWebBrowser::UnbindUObject(const FString& Name, UObject* Object, bool bIsPermanent)
+{
+	if (BrowserWindow.IsValid())
+	{
+		BrowserWindow->UnbindUObject(Name, Object, bIsPermanent);
+	}
+}
+
 #undef LOCTEXT_NAMESPACE
 
diff --git a/Engine/Source/Runtime/WebBrowser/Private/WebBrowserApp.cpp b/Engine/Source/Runtime/WebBrowser/Private/WebBrowserApp.cpp
--- a/Engine/Source/Runtime/WebBrowser/Private/WebBrowserApp.cpp
+++ b/Engine/Source/Runtime/WebBrowser/Private/WebBrowserApp.cpp
@@ -8,27 +8,13 @@ FWebBrowserApp::FWebBrowserApp()
 {
 }
 
-void FWebBrowserApp::OnContextInitialized()
-{
-	// Called on the browser process UI thread immediately after the CEF context
-	// has been initialized.
-}
-
 void FWebBrowserApp::OnBeforeChildProcessLaunch(CefRefPtr<CefCommandLine> CommandLine)
 {
-	std::string x = CommandLine->GetCommandLineString();
-	std::string a = "d"+x;
 }
 
-
-/*void FWebBrowserApp::OnContextCreated(CefRefPtr<CefBrowser> Browser, CefRefPtr<CefFrame> Frame, CefRefPtr<CefV8Context> Context)
+void FWebBrowserApp::OnRenderProcessThreadCreated(CefRefPtr<CefListValue> ExtraInfo)
 {
-	// Retrieve the context's window object.
-	CefRefPtr<CefV8Value> Object = Context->GetGlobal();
-	
-	Object->SetValue("answer",
-					 CefV8Value::CreateUInt(42),
-					 V8_PROPERTY_ATTRIBUTE_READONLY);
-}*/
+	RenderProcessThreadCreatedDelegate.ExecuteIfBound(ExtraInfo);
+}
 
 #endif
diff --git a/Engine/Source/Runtime/WebBrowser/Private/WebBrowserApp.h b/Engine/Source/Runtime/WebBrowser/Private/WebBrowserApp.h
--- a/Engine/Source/Runtime/WebBrowser/Private/WebBrowserApp.h
+++ b/Engine/Source/Runtime/WebBrowser/Private/WebBrowserApp.h
@@ -27,13 +27,23 @@ public:
 	 * Default Constructor
 	 */
 	FWebBrowserApp();
+
+	/** A delegate this is invoked when an existing browser requests creation of a new browser window. */
+	DECLARE_DELEGATE_OneParam(FOnRenderProcessThreadCreated, CefRefPtr<CefListValue> /*ExtraInfo*/);
+	virtual FOnRenderProcessThreadCreated& OnRenderProcessThreadCreated()
+	{
+		return RenderProcessThreadCreatedDelegate;
+	}
+
 private:
 	// CefApp methods.
 	virtual CefRefPtr<CefBrowserProcessHandler> GetBrowserProcessHandler() override { return this; }
 
 	// CefBrowserProcessHandler methods:
-	virtual void OnContextInitialized() override;
 	virtual void OnBeforeChildProcessLaunch(CefRefPtr<CefCommandLine> CommandLine) override;
+	virtual void OnRenderProcessThreadCreated(CefRefPtr<CefListValue> ExtraInfo) override;
+
+	FOnRenderProcessThreadCreated RenderProcessThreadCreatedDelegate;
 
 	// Include the default reference counting implementation.
 	IMPLEMENT_REFCOUNTING(FWebBrowserApp);
diff --git a/Engine/Source/Runtime/WebBrowser/Private/WebBrowserHandler.cpp b/Engine/Source/Runtime/WebBrowser/Private/WebBrowserHandler.cpp
--- a/Engine/Source/Runtime/WebBrowser/Private/WebBrowserHandler.cpp
+++ b/Engine/Source/Runtime/WebBrowser/Private/WebBrowserHandler.cpp
@@ -193,13 +193,17 @@ void FWebBrowserHandler::OnLoadError(CefRefPtr<CefBrowser> Browser,
 	}
 }
 
-void FWebBrowserHandler::OnLoadingStateChange(CefRefPtr<CefBrowser> browser, bool isLoading, bool canGoBack, bool canGoForward)
+void FWebBrowserHandler::OnLoadStart(CefRefPtr<CefBrowser> Browser, CefRefPtr<CefFrame> Frame)
+{
+}
+
+void FWebBrowserHandler::OnLoadingStateChange(CefRefPtr<CefBrowser> Browser, bool bIsLoading, bool bCanGoBack, bool bCanGoForward)
 {
 	TSharedPtr<FWebBrowserWindow> BrowserWindow = BrowserWindowPtr.Pin();
 
 	if (BrowserWindow.IsValid())
 	{
-		BrowserWindow->NotifyDocumentLoadingStateChange(isLoading);
+		BrowserWindow->NotifyDocumentLoadingStateChange(bIsLoading);
 	}
 }
 
@@ -307,7 +311,17 @@ bool FWebBrowserHandler::OnProcessMessageReceived(CefRefPtr<CefBrowser> Browser,
 	CefProcessId SourceProcess,
 	CefRefPtr<CefProcessMessage> Message)
 {
-	return MessageRouter->OnProcessMessageReceived(Browser, SourceProcess, Message);
+	bool Retval = false;
+	TSharedPtr<FWebBrowserWindow> BrowserWindow = BrowserWindowPtr.Pin();
+	if (BrowserWindow.IsValid())
+	{
+		Retval = BrowserWindow->OnProcessMessageReceived(Browser, SourceProcess, Message);
+	}
+	if (! Retval)
+	{
+		Retval = MessageRouter->OnProcessMessageReceived(Browser, SourceProcess, Message);
+	}
+	return Retval;
 }
 
 bool FWebBrowserHandler::OnQuery(CefRefPtr<CefBrowser> Browser,
diff --git a/Engine/Source/Runtime/WebBrowser/Private/WebBrowserHandler.h b/Engine/Source/Runtime/WebBrowser/Private/WebBrowserHandler.h
--- a/Engine/Source/Runtime/WebBrowser/Private/WebBrowserHandler.h
+++ b/Engine/Source/Runtime/WebBrowser/Private/WebBrowserHandler.h
@@ -142,6 +142,9 @@ public:
 		bool canGoBack,
 		bool canGoForward) override;
 
+	virtual void OnLoadStart(CefRefPtr<CefBrowser> Browser, CefRefPtr<CefFrame> Frame) override;
+
+
 public:
 
 	// CefRenderHandler Interface
diff --git a/Engine/Source/Runtime/WebBrowser/Private/WebBrowserSingleton.cpp b/Engine/Source/Runtime/WebBrowser/Private/WebBrowserSingleton.cpp
--- a/Engine/Source/Runtime/WebBrowser/Private/WebBrowserSingleton.cpp
+++ b/Engine/Source/Runtime/WebBrowser/Private/WebBrowserSingleton.cpp
@@ -37,6 +37,7 @@ FWebBrowserSingleton::FWebBrowserSingleton()
 
 	// WebBrowserApp implements application-level callbacks.
 	WebBrowserApp = new FWebBrowserApp;
+	WebBrowserApp->OnRenderProcessThreadCreated().BindRaw(this, &FWebBrowserSingleton::HandleRenderProcessCreated);
 
 	// Specify CEF global settings here.
 	CefSettings Settings;
@@ -102,6 +103,24 @@ FWebBrowserSingleton::FWebBrowserSingleton()
 #endif
 }
 
+#if WITH_CEF3
+void FWebBrowserSingleton::HandleRenderProcessCreated(CefRefPtr<CefListValue> ExtraInfo)
+{
+	for (int32 Index = WindowInterfaces.Num() - 1; Index >= 0; --Index)
+	{
+		TSharedPtr<FWebBrowserWindow> BrowserWindow = WindowInterfaces[Index].Pin();
+		if (BrowserWindow.IsValid())
+		{
+			CefRefPtr<CefDictionaryValue> Bindings = BrowserWindow->GetProcessInfo();
+			if (Bindings.get())
+			{
+				ExtraInfo->SetDictionary(ExtraInfo->GetSize(), Bindings);
+			}
+		}
+	}
+}
+#endif
+
 FWebBrowserSingleton::~FWebBrowserSingleton()
 {
 #if WITH_CEF3
@@ -113,7 +132,8 @@ FWebBrowserSingleton::~FWebBrowserSingleton()
 			WindowInterfaces[Index].Pin()->CloseBrowser();
 		}
 	}
-
+	// Just in case, although we deallocate WebBrowserApp right after this.
+	WebBrowserApp->OnRenderProcessThreadCreated().Unbind();
 	// CefRefPtr takes care of delete
 	WebBrowserApp = nullptr;
 	// Shut down CEF.
diff --git a/Engine/Source/Runtime/WebBrowser/Private/WebBrowserSingleton.h b/Engine/Source/Runtime/WebBrowser/Private/WebBrowserSingleton.h
--- a/Engine/Source/Runtime/WebBrowser/Private/WebBrowserSingleton.h
+++ b/Engine/Source/Runtime/WebBrowser/Private/WebBrowserSingleton.h
@@ -15,6 +15,8 @@
 #if PLATFORM_WINDOWS
 	#include "HideWindowsPlatformTypes.h"
 #endif
+
+class CefListValue;
 #endif
 
 
@@ -73,6 +75,8 @@ public:
 
 private:
 #if WITH_CEF3
+	/** When new render processes are created, send all permanent variable bindings to them. */
+	void HandleRenderProcessCreated(CefRefPtr<CefListValue> ExtraInfo);
 	/** Pointer to the CEF App implementation */
 	CefRefPtr<FWebBrowserApp>			WebBrowserApp;
 	/** List of currently existing browser windows */
diff --git a/Engine/Source/Runtime/WebBrowser/Private/WebBrowserWindow.cpp b/Engine/Source/Runtime/WebBrowser/Private/WebBrowserWindow.cpp
--- a/Engine/Source/Runtime/WebBrowser/Private/WebBrowserWindow.cpp
+++ b/Engine/Source/Runtime/WebBrowser/Private/WebBrowserWindow.cpp
@@ -3,6 +3,7 @@
 #include "WebBrowserPrivatePCH.h"
 #include "WebBrowserWindow.h"
 #include "WebBrowserByteResource.h"
+#include "WebJSScripting.h"
 #include "RHI.h"
 
 #if WITH_CEF3
@@ -33,6 +34,7 @@ FWebBrowserWindow::FWebBrowserWindow(FIntPoint InViewportSize, FString InUrl, TO
 	, bIgnoreKeyDownEvent(false)
 	, bIgnoreKeyUpEvent(false)
 	, bIgnoreCharacterEvent(false)
+	, Scripting(new FWebJSScripting)
 {
 }
 
@@ -672,6 +674,7 @@ void FWebBrowserWindow::CloseBrowser()
 	bIsClosing = true;
 	if (IsValid())
 	{
+		Scripting->UnbindCefBrowser();
 		InternalCefBrowser->GetHost()->CloseBrowser(false);
 		InternalCefBrowser = nullptr;
 		Handler = nullptr;
@@ -681,6 +684,7 @@ void FWebBrowserWindow::CloseBrowser()
 void FWebBrowserWindow::BindCefBrowser(CefRefPtr<CefBrowser> Browser)
 {
 	check(Browser.get() == nullptr || InternalCefBrowser.get() == nullptr || InternalCefBrowser->IsSame(Browser));
+	Scripting->BindCefBrowser(Browser); // The scripting interface needs the browser too
 	InternalCefBrowser = Browser;
 }
 
@@ -689,7 +693,6 @@ CefRefPtr<CefBrowser> FWebBrowserWindow::GetCefBrowser()
 	return InternalCefBrowser;
 }
 
-
 void FWebBrowserWindow::SetTitle(const CefString& InTitle)
 {
 	Title = InTitle.ToWString().c_str();
@@ -1018,9 +1021,36 @@ void FWebBrowserWindow::SetIsHidden(bool bValue)
 	}
 }
 
+CefRefPtr<CefDictionaryValue> FWebBrowserWindow::GetProcessInfo()
+{
+	CefRefPtr<CefDictionaryValue> Retval = nullptr;
+	if (IsValid())
+	{
+		Retval = CefDictionaryValue::Create();
+		Retval->SetInt("browser", InternalCefBrowser->GetIdentifier());
+		Retval->SetDictionary("bindings", Scripting->GetPermanentBindings());
+	}
+	return Retval;
+}
+
+bool FWebBrowserWindow::OnProcessMessageReceived(CefRefPtr<CefBrowser> Browser, CefProcessId SourceProcess, CefRefPtr<CefProcessMessage> Message)
+{
+	return Scripting->OnProcessMessageReceived(Browser, SourceProcess, Message);
+}
+
+void FWebBrowserWindow::BindUObject(const FString& Name, UObject* Object, bool bIsPermanent)
+{
+	Scripting->BindUObject(Name, Object, bIsPermanent);
+}
+
+void FWebBrowserWindow::UnbindUObject(const FString& Name, UObject* Object, bool bIsPermanent)
+{
+	Scripting->UnbindUObject(Name, Object, bIsPermanent);
+}
+
 bool FWebBrowserWindow::OnQuery(int64 QueryId, const CefString& Request, bool Persistent, CefRefPtr<CefMessageRouterBrowserSide::Callback> Callback)
 {
-	if ( OnJSQueryReceived().IsBound() )
+	if (OnJSQueryReceived().IsBound())
 	{
 		FString QueryString = Request.ToWString().c_str();
 		FJSQueryResultDelegate Delegate = FJSQueryResultDelegate::CreateLambda(
diff --git a/Engine/Source/Runtime/WebBrowser/Private/WebBrowserWindow.h b/Engine/Source/Runtime/WebBrowser/Private/WebBrowserWindow.h
--- a/Engine/Source/Runtime/WebBrowser/Private/WebBrowserWindow.h
+++ b/Engine/Source/Runtime/WebBrowser/Private/WebBrowserWindow.h
@@ -6,8 +6,10 @@
 
 #include "IWebBrowserWindow.h"
 #include "WebBrowserHandler.h"
+#include "WebJSScripting.h"
 #include "SlateCore.h"
 #include "SlateBasics.h"
+#include "CoreUObject.h"
 
 #if PLATFORM_WINDOWS
 	#include "AllowWindowsPlatformTypes.h"
@@ -114,6 +116,8 @@ public:
 	virtual void StopLoad() override;
 	virtual void ExecuteJavascript(const FString& Script) override;
 	virtual void CloseBrowser() override;
+	virtual void BindUObject(const FString& Name, UObject* Object, bool bIsPermanent = true) override;
+	virtual void UnbindUObject(const FString& Name, UObject* Object = nullptr, bool bIsPermanent = true) override;
 
 	DECLARE_DERIVED_EVENT(FWebBrowserWindow, IWebBrowserWindow::FOnDocumentStateChanged, FOnDocumentStateChanged);
 	virtual FOnDocumentStateChanged& OnDocumentStateChanged() override
@@ -246,19 +250,26 @@ private:
 	 */
 	void OnCursorChange(CefCursorHandle Cursor, CefRenderHandler::CursorType Type, const CefCursorInfo& CustomCursorInfo);
 
+	/**
+	 * Called when a message was received from the renderer process.
+	 *
+	 * @param Browser The CefBrowser for this window.
+	 * @param SourceProcess The process id of the sender of the message. Currently always PID_RENDERER.
+	 * @param Message The actual message.
+	 * @return true if the message was handled, else false.
+	 */
+	bool OnProcessMessageReceived(CefRefPtr<CefBrowser> Browser, CefProcessId SourceProcess, CefRefPtr<CefProcessMessage> Message);
+
 	/**
 	 * Called when JavaScript code sends a message to the UE process.
 	 * Needs to return true or false to tell CEF wether the query is being handled by user code or not.
 	 *
 	 * @param QueryId A unique id for the query. Used to refer to it in OnQueryCanceled.
 	 * @param Request The query string itself as passed in from the JS code.
-	 * @param Persistent Os this a persistent query or not. If not, client code expects the callback to be invoked only once, wheras persistent queries are terminated by invoking Failure, Success can be invoked multiple times until then.
+	 * @param Persistent Is this a persistent query or not. If not, client code expects the callback to be invoked only once, wheras persistent queries are terminated by invoking Failure, Success can be invoked multiple times until then.
 	 * @param Callback A handle to pass data back to the JS code. 
 	 */
-	bool OnQuery(int64 QueryId,
-		const CefString& Request,
-		bool Persistent,
-		CefRefPtr<CefMessageRouterBrowserSide::Callback> Callback);
+	bool OnQuery(int64 QueryId, const CefString& Request, bool Persistent, CefRefPtr<CefMessageRouterBrowserSide::Callback> Callback);
 
 	/**
 	 * Called when an outstanding query has been canceled eother explicitly from JS code or implicitly by navigating away from the page containing the code.
@@ -339,12 +350,19 @@ public:
 	 */
 	void CheckTickActivity();
 
+	/**
+	 * Called on every browser window when CEF launches a new render process.
+	 * Used to ensure global JS objects are registered as soon as possible.
+	 */
+	CefRefPtr<CefDictionaryValue> GetProcessInfo();
+
 private:
 	/** Helper that calls WasHidden on the CEF host object when the value changes */
 	void SetIsHidden(bool bValue);
 
 	/** Used by the key down and up handlers to convert Slate key events to the CEF equivalent. */
 	void PopulateCefKeyEvent(const FKeyEvent& InKeyEvent, CefKeyEvent& OutKeyEvent);
+
 private:
 
 	/** Current state of the document being loaded. */
@@ -435,6 +453,9 @@ private:
 	bool bIgnoreKeyDownEvent;
 	bool bIgnoreKeyUpEvent;
 	bool bIgnoreCharacterEvent;
+
+	/** Handling of passing and marshalling messages for JS integration is delegated to a helper class*/
+	TSharedPtr<FWebJSScripting> Scripting;
 };
 
 #endif
diff --git a/Engine/Source/Runtime/WebBrowser/Private/WebJSFunction.cpp b/Engine/Source/Runtime/WebBrowser/Private/WebJSFunction.cpp
new file mode 100644
--- /dev/null
+++ b/Engine/Source/Runtime/WebBrowser/Private/WebJSFunction.cpp
@@ -0,0 +1,67 @@
+// Copyright 1998-2015 Epic Games, Inc. All Rights Reserved.
+
+#include "WebBrowserPrivatePCH.h"
+#include "WebJSFunction.h"
+#include "WebJSScripting.h"
+
+#if WITH_CEF3
+#if PLATFORM_WINDOWS
+#include "AllowWindowsPlatformTypes.h"
+#endif
+#pragma push_macro("OVERRIDE")
+#undef OVERRIDE // cef headers provide their own OVERRIDE macro
+#include "include/cef_values.h"
+#pragma pop_macro("OVERRIDE")
+#if PLATFORM_WINDOWS
+#include "HideWindowsPlatformTypes.h"
+#endif
+#endif
+
+FWebJSParam::~FWebJSParam()
+{
+	// Since the FString member is in a union, it may or may not be valid, so we have to call its destructor manually.
+	if (Tag == PTYPE_STRING)
+	{
+		delete StringValue;
+	}
+}
+
+FWebJSParam::FWebJSParam(const FWebJSParam& Other)
+	: Tag(Other.Tag)
+{
+	switch (Other.Tag)
+	{
+		case PTYPE_BOOL:
+			BoolValue = Other.BoolValue;
+			break;
+		case PTYPE_DOUBLE:
+			DoubleValue = Other.DoubleValue;
+			break;
+		case PTYPE_INT:
+			IntValue = Other.IntValue;
+			break;
+		case PTYPE_STRING:
+			StringValue = new FString(*Other.StringValue);
+			break;
+		case PTYPE_NULL:
+			break;
+		case PTYPE_OBJECT:
+			ObjectValue = Other.ObjectValue;
+			break;
+		case PTYPE_STRUCT:
+			StructValue.TypeInfo = Other.StructValue.TypeInfo;
+			StructValue.StructPtr = Other.StructValue.StructPtr;
+			break;
+	}
+}
+
+void FWebJSFunction::Invoke(int32 ArgCount, FWebJSParam Arguments[]) const
+{
+#if WITH_CEF3
+	TSharedPtr<FWebJSScripting> Scripting = ScriptingPtr.Pin();
+	if (Scripting.IsValid())
+	{
+		Scripting->InvokeJSFunction(FunctionId, ArgCount, Arguments);
+	}
+#endif
+}
\ No newline at end of file
diff --git a/Engine/Source/Runtime/WebBrowser/Private/WebJSScripting.cpp b/Engine/Source/Runtime/WebBrowser/Private/WebJSScripting.cpp
new file mode 100644
--- /dev/null
+++ b/Engine/Source/Runtime/WebBrowser/Private/WebJSScripting.cpp
@@ -0,0 +1,432 @@
+
+ // Copyright 1998-2015 Epic Games, Inc. All Rights Reserved.
+
+#include "WebBrowserPrivatePCH.h"
+#include "WebJSScripting.h"
+#include "WebJSFunction.h"
+#include "WebBrowserWindow.h"
+#include "WebJSStructSerializerBackend.h"
+#include "WebJSStructDeserializerBackend.h"
+#include "StructSerializer.h"
+#include "StructDeserializer.h"
+
+#if WITH_CEF3
+
+// Internal utility function(s)
+namespace
+{
+
+	template<typename DestContainerType, typename SrcContainerType, typename DestKeyType, typename SrcKeyType>
+	bool CopyContainerValue(DestContainerType DestContainer, SrcContainerType SrcContainer, DestKeyType DestKey, SrcKeyType SrcKey )
+	{
+		switch (SrcContainer->GetType(SrcKey))
+		{
+			case VTYPE_NULL:
+				return DestContainer->SetNull(DestKey);
+			case VTYPE_BOOL:
+				return DestContainer->SetBool(DestKey, SrcContainer->GetBool(SrcKey));
+			case VTYPE_INT:
+				return DestContainer->SetInt(DestKey, SrcContainer->GetInt(SrcKey));
+			case VTYPE_DOUBLE:
+				return DestContainer->SetDouble(DestKey, SrcContainer->GetDouble(SrcKey));
+			case VTYPE_STRING:
+				return DestContainer->SetString(DestKey, SrcContainer->GetString(SrcKey));
+			case VTYPE_BINARY:
+				return DestContainer->SetBinary(DestKey, SrcContainer->GetBinary(SrcKey));
+			case VTYPE_DICTIONARY:
+				return DestContainer->SetDictionary(DestKey, SrcContainer->GetDictionary(SrcKey));
+			case VTYPE_LIST:
+				return DestContainer->SetList(DestKey, SrcContainer->GetList(SrcKey));
+			case VTYPE_INVALID:
+			default:
+				return false;
+		}
+
+	}
+}
+
+FWebJSParam FWebJSScripting::ConvertResult(UProperty* Property, uint8* Data)
+{
+	// booleans
+	if (Property->IsA<UBoolProperty>())
+	{
+		return FWebJSParam(Cast<UBoolProperty>(Property)->GetPropertyValue_InContainer(Data));
+	}
+
+	// unsigned bytes & enumerations
+	else if (Property->IsA<UByteProperty>())
+	{
+		UByteProperty* ByteProperty = Cast<UByteProperty>(Property);
+
+		if (ByteProperty->IsEnum())
+		{
+			return FWebJSParam(ByteProperty->Enum->GetEnumName(ByteProperty->GetPropertyValue_InContainer(Data)));
+		}
+		else
+		{
+			return FWebJSParam((int32)ByteProperty->GetPropertyValue_InContainer(Data));
+		}
+	}
+
+	// floating point numbers
+	else if (Property->IsA<UDoubleProperty>())
+	{
+		return FWebJSParam(Cast<UDoubleProperty>(Property)->GetPropertyValue_InContainer(Data));
+	}
+	else if (Property->IsA<UFloatProperty>())
+	{
+		return FWebJSParam(Cast<UFloatProperty>(Property)->GetPropertyValue_InContainer(Data));
+	}
+
+	// signed integers
+	else if (Property->IsA<UIntProperty>())
+	{
+		return FWebJSParam(Cast<UIntProperty>(Property)->GetPropertyValue_InContainer(Data));
+	}
+	else if (Property->IsA<UInt8Property>())
+	{
+		return FWebJSParam((int32)Cast<UInt8Property>(Property)->GetPropertyValue_InContainer(Data));
+	}
+	else if (Property->IsA<UInt16Property>())
+	{
+		return FWebJSParam((int32)Cast<UInt16Property>(Property)->GetPropertyValue_InContainer(Data));
+	}
+	else if (Property->IsA<UInt64Property>())
+	{
+		return FWebJSParam((double)Cast<UInt64Property>(Property)->GetPropertyValue_InContainer(Data));
+	}
+
+	// unsigned integers
+	else if (Property->IsA<UUInt16Property>())
+	{
+		return FWebJSParam((int32)Cast<UUInt16Property>(Property)->GetPropertyValue_InContainer(Data));
+	}
+	else if (Property->IsA<UUInt32Property>())
+	{
+		return FWebJSParam((double)Cast<UUInt32Property>(Property)->GetPropertyValue_InContainer(Data));
+	}
+	else if (Property->IsA<UUInt64Property>())
+	{
+		return FWebJSParam((double)Cast<UUInt64Property>(Property)->GetPropertyValue_InContainer(Data));
+	}
+
+	// names & strings
+	else if (Property->IsA<UNameProperty>())
+	{
+		return FWebJSParam(Cast<UNameProperty>(Property)->GetPropertyValue_InContainer(Data).ToString());
+	}
+	else if (Property->IsA<UStrProperty>())
+	{
+		return FWebJSParam(Cast<UStrProperty>(Property)->GetPropertyValue_InContainer(Data));
+	}
+
+	// classes & objects
+	else if (Property->IsA<UClassProperty>())
+	{
+		return FWebJSParam(Cast<UClassProperty>(Property)->GetPropertyValue_InContainer(Data)->GetPathName());
+	}
+	else if (Property->IsA<UObjectProperty>())
+	{
+		return FWebJSParam(Cast<UObjectProperty>(Property)->GetPropertyValue_InContainer(Data));
+	}
+	else if (Property->IsA<UStructProperty>())
+	{
+		UStructProperty* StructProperty = Cast<UStructProperty>(Property);
+		return FWebJSParam(StructProperty->Struct, StructProperty->ContainerPtrToValuePtr<void>(Data));
+	}
+	else
+	{
+		GLog->Logf(ELogVerbosity::Warning, TEXT("FWebJSScripting: %s cannot be serialized, because its type (%s) is not supported"), *Property->GetName(), *Property->GetClass()->GetName());
+	}
+	return FWebJSParam();
+}
+
+
+CefRefPtr<CefDictionaryValue> FWebJSScripting::ConvertStruct(UStruct* TypeInfo, void* StructPtr)
+{
+	FWebJSStructSerializerBackend Backend (SharedThis(this));
+	FStructSerializer::Serialize(StructPtr, *TypeInfo, Backend);
+
+	CefRefPtr<CefDictionaryValue> Result = CefDictionaryValue::Create();
+	Result->SetString("$type", "struct");
+	Result->SetString("$ue4Type", *TypeInfo->GetName());
+	Result->SetDictionary("$value", Backend.GetResult());
+	return Result;
+}
+
+CefRefPtr<CefDictionaryValue> FWebJSScripting::ConvertObject(UObject* Object)
+{
+	CefRefPtr<CefDictionaryValue> Result = CefDictionaryValue::Create();
+	RetainBinding(Object);
+
+	UClass* Class = Object->GetClass();
+	CefRefPtr<CefListValue> MethodNames = CefListValue::Create();
+	int32 MethodIndex = 0;
+	for (TFieldIterator<UFunction> FunctionIt(Class, EFieldIteratorFlags::IncludeSuper); FunctionIt; ++FunctionIt)
+	{
+		UFunction* Function = *FunctionIt;
+		MethodNames->SetString(MethodIndex++, *Function->GetName());
+	}
+
+	Result->SetString("$type", "uobject");
+	Result->SetString("$id", *PtrToGuid(Object).ToString(EGuidFormats::Digits));
+	Result->SetList("$methods", MethodNames);
+	return Result;
+}
+
+
+bool FWebJSScripting::OnProcessMessageReceived(CefRefPtr<CefBrowser> Browser, CefProcessId SourceProcess, CefRefPtr<CefProcessMessage> Message)
+{
+	bool Result = false;
+	FString MessageName = Message->GetName().ToWString().c_str();
+	if (MessageName == TEXT("UE::ExecuteUObjectMethod"))
+	{
+		Result = HandleExecuteUObjectMethodMessage(Message->GetArgumentList());
+	}
+	else if (MessageName == TEXT("UE::ReleaseUObject"))
+	{
+		Result = HandleReleaseUObjectMessage(Message->GetArgumentList());
+	}
+	return Result;
+}
+
+void FWebJSScripting::SendProcessMessage(CefRefPtr<CefProcessMessage> Message)
+{
+	if (IsValid() )
+	{
+		InternalCefBrowser->SendProcessMessage(PID_RENDERER, Message);
+	}
+}
+
+CefRefPtr<CefDictionaryValue> FWebJSScripting::GetPermanentBindings()
+{
+	CefRefPtr<CefDictionaryValue> Result = CefDictionaryValue::Create();
+	for(auto& Entry : PermanentUObjectsByName)
+	{
+		Result->SetDictionary(*Entry.Key, ConvertObject(Entry.Value));
+	}
+	return Result;
+}
+
+
+void FWebJSScripting::BindUObject(const FString& Name, UObject* Object, bool bIsPermanent)
+{
+	CefRefPtr<CefDictionaryValue> Converted = ConvertObject(Object);
+	if (bIsPermanent)
+	{
+		// Each object can only have one permanent binding
+		if (BoundObjects[Object].bIsPermanent)
+		{
+			return;
+		}
+		// Existing permanent objects must be removed first
+		if (PermanentUObjectsByName.Contains(Name))
+		{
+			return;
+		}
+		BoundObjects[Object]={true, -1};
+		PermanentUObjectsByName.Add(Name, Object);
+	}
+
+	CefRefPtr<CefProcessMessage> SetValueMessage = CefProcessMessage::Create(TEXT("UE::SetValue"));
+	CefRefPtr<CefListValue>MessageArguments = SetValueMessage->GetArgumentList();
+	CefRefPtr<CefDictionaryValue> Value = CefDictionaryValue::Create();
+	Value->SetString("name", *Name);
+	Value->SetDictionary("value", Converted);
+	Value->SetBool("permanent", bIsPermanent);
+
+	MessageArguments->SetDictionary(0, Value);
+	SendProcessMessage(SetValueMessage);
+}
+
+void FWebJSScripting::UnbindUObject(const FString& Name, UObject* Object, bool bIsPermanent)
+{
+	if (bIsPermanent)
+	{
+		// If overriding an existing permanent object, make it non-permanent
+		if (PermanentUObjectsByName.Contains(Name) && (Object == nullptr || PermanentUObjectsByName[Name] == Object))
+		{
+			Object = PermanentUObjectsByName.FindAndRemoveChecked(Name);
+			BoundObjects.Remove(Object);
+			return;
+		}
+		else
+		{
+			return;
+		}
+	}
+
+	CefRefPtr<CefProcessMessage> DeleteValueMessage = CefProcessMessage::Create(TEXT("UE::DeleteValue"));
+	CefRefPtr<CefListValue>MessageArguments = DeleteValueMessage->GetArgumentList();
+	CefRefPtr<CefDictionaryValue> Info = CefDictionaryValue::Create();
+	Info->SetString("name", *Name);
+	Info->SetString("id", *PtrToGuid(Object).ToString(EGuidFormats::Digits));
+	Info->SetBool("permanent", bIsPermanent);
+
+	MessageArguments->SetDictionary(0, Info);
+	SendProcessMessage(DeleteValueMessage);
+}
+
+bool FWebJSScripting::HandleReleaseUObjectMessage(CefRefPtr<CefListValue> MessageArguments)
+{
+	FGuid ObjectKey;
+	// Message arguments are Name, Value, bGlobal
+	if (MessageArguments->GetSize() != 1 || MessageArguments->GetType(0) != VTYPE_STRING)
+	{
+		// Wrong message argument types or count
+		return false;
+	}
+
+	if (!FGuid::Parse(FString(MessageArguments->GetString(0).ToWString().c_str()), ObjectKey))
+	{
+		// Invalid GUID
+		return false;
+	}
+
+	UObject* Object = GuidToPtr(ObjectKey);
+	if ( Object == nullptr )
+	{
+		// Invalid object
+		return false;
+	}
+	ReleaseBinding(Object);
+	return true;
+}
+
+bool FWebJSScripting::HandleExecuteUObjectMethodMessage(CefRefPtr<CefListValue> MessageArguments)
+{
+	FGuid ObjectKey;
+	// Message arguments are Name, Value, bGlobal
+	if (MessageArguments->GetSize() != 4
+		|| MessageArguments->GetType(0) != VTYPE_STRING
+		|| MessageArguments->GetType(1) != VTYPE_STRING
+		|| MessageArguments->GetType(2) != VTYPE_STRING
+		|| MessageArguments->GetType(3) != VTYPE_LIST
+		)
+	{
+		// Wrong message argument types or count
+		return false;
+	}
+
+	if (!FGuid::Parse(FString(MessageArguments->GetString(0).ToWString().c_str()), ObjectKey))
+	{
+		// Invalid GUID
+		return false;
+	}
+
+	// Get the promise callback and use that to report any results from executing this function.
+	FGuid ResultCallbackId;
+	if (!FGuid::Parse(FString(MessageArguments->GetString(2).ToWString().c_str()), ResultCallbackId))
+	{
+		// Invalid GUID
+		return false;
+	}
+
+	UObject* Object = GuidToPtr(ObjectKey);
+	if (Object == nullptr)
+	{
+		// Unknown uobject id
+		InvokeJSErrorResult(ResultCallbackId, TEXT("Unknown UObject ID"));
+		return true;
+	}
+
+	FName MethodName = MessageArguments->GetString(1).ToWString().c_str();
+	UFunction* Function = Object->FindFunction(MethodName);
+	if (!Function)
+	{
+		InvokeJSErrorResult(ResultCallbackId, TEXT("Unknown UObject Function"));
+		return true;
+	}
+	// Coerce arguments to function arguments.
+	uint16 ParamsSize = Function->ParmsSize;
+	TArray<uint8> Params;
+	UProperty* ReturnParam = nullptr;
+
+	// Convert cef argument list to a dictionary, so we can use FStructDeserializer to convert it for us
+	CefRefPtr<CefDictionaryValue> NamedArgs = CefDictionaryValue::Create();
+	{
+		int32 CurrentArg = 0;
+		CefRefPtr<CefListValue> CefArgs = MessageArguments->GetList(3);
+		for ( TFieldIterator<UProperty> It(Function); It; ++It )
+		{
+			UProperty* Param = *It;
+			if (Param->PropertyFlags & CPF_Parm)
+			{
+				if (Param->PropertyFlags & CPF_ReturnParm)
+				{
+					ReturnParam = Param;
+				}
+				else
+				{
+					CopyContainerValue(NamedArgs, CefArgs, CefString(*Param->GetName()), CurrentArg);
+					CurrentArg++;
+				}
+			}
+		}
+
+		if (ParamsSize > 0)
+		{
+			Params.AddZeroed(ParamsSize);
+			// UFunction is a subclass of UStruct, so we can treat the arguments as a struct for deserialization
+			FWebJSStructDeserializerBackend Backend = FWebJSStructDeserializerBackend(SharedThis(this), NamedArgs);
+			FStructDeserializer::Deserialize(Params.GetData(), *Cast<UStruct>(Function), Backend);
+		}
+	}
+
+	Object->ProcessEvent(Function, Params.GetData());
+	if ( ReturnParam )
+	{
+		FWebJSParam Results[1] = {ConvertResult(ReturnParam, Params.GetData())};
+		InvokeJSFunction(ResultCallbackId, 1, Results, false);
+	}
+	else
+	{
+		InvokeJSFunction(ResultCallbackId, 0, nullptr, false);
+	}
+	return true;
+}
+
+void FWebJSScripting::BindCefBrowser(CefRefPtr<CefBrowser> Browser)
+{
+	check(Browser.get() == nullptr || InternalCefBrowser.get() == nullptr || InternalCefBrowser->IsSame(Browser));
+	InternalCefBrowser = Browser;
+}
+
+void FWebJSScripting::UnbindCefBrowser()
+{
+	InternalCefBrowser = nullptr;
+}
+
+
+void FWebJSScripting::AddReferencedObjects(FReferenceCollector& Collector)
+{
+	// Ensure bound UObjects are not garbage collected as long as this object is valid.
+	for (auto& Binding : BoundObjects)
+	{
+		Collector.AddReferencedObject(Binding.Key);
+	}
+}
+
+void FWebJSScripting::InvokeJSErrorResult(FGuid FunctionId, const FString& Error)
+{
+	FWebJSParam Arguments[1] = {FWebJSParam(Error)};
+	InvokeJSFunction(FunctionId, 1, Arguments, true);
+}
+
+void FWebJSScripting::InvokeJSFunction(FGuid FunctionId, int32 ArgCount, FWebJSParam Arguments[], bool bIsError)
+{
+	CefRefPtr<CefListValue> FunctionArguments = CefListValue::Create();
+	for ( int32 i=0; i<ArgCount; i++)
+	{
+		SetConverted(FunctionArguments, i, Arguments[i]);
+	}
+
+	CefRefPtr<CefProcessMessage> Message = CefProcessMessage::Create(TEXT("UE::ExecuteJSFunction"));
+	CefRefPtr<CefListValue> MessageArguments = Message->GetArgumentList();
+	MessageArguments->SetString(0, *FunctionId.ToString(EGuidFormats::Digits));
+	MessageArguments->SetList(1, FunctionArguments);
+	MessageArguments->SetBool(2, bIsError);
+	SendProcessMessage(Message);
+}
+
+#endif
\ No newline at end of file
diff --git a/Engine/Source/Runtime/WebBrowser/Private/WebJSScripting.h b/Engine/Source/Runtime/WebBrowser/Private/WebJSScripting.h
new file mode 100644
--- /dev/null
+++ b/Engine/Source/Runtime/WebBrowser/Private/WebJSScripting.h
@@ -0,0 +1,211 @@
+#pragma once
+#include "CoreUObject.h"
+#include "WebJSFunction.h"
+
+#if WITH_CEF3
+#if PLATFORM_WINDOWS
+#include "AllowWindowsPlatformTypes.h"
+#endif
+#pragma push_macro("OVERRIDE")
+#undef OVERRIDE // cef headers provide their own OVERRIDE macro
+#include "include/cef_client.h"
+#include "include/cef_values.h"
+#pragma pop_macro("OVERRIDE")
+#if PLATFORM_WINDOWS
+#include "HideWindowsPlatformTypes.h"
+#endif
+
+/**
+ * Implements handling of bridging UObjects client side with JavaScript renderer side.
+ */
+class FWebJSScripting
+	: public FGCObject
+	, public TSharedFromThis<FWebJSScripting>
+{
+public:
+	FWebJSScripting()
+		: BaseGuid(FGuid::NewGuid())
+	{}
+
+	void BindUObject(const FString& Name, UObject* Object, bool bIsPermanent = true);
+	void UnbindUObject(const FString& Name, UObject* Object = nullptr, bool bIsPermanent = true);
+
+	/**
+	 * Called when a message was received from the renderer process.
+	 *
+	 * @param Browser The CefBrowser for this window.
+	 * @param SourceProcess The process id of the sender of the message. Currently always PID_RENDERER.
+	 * @param Message The actual message.
+	 * @return true if the message was handled, else false.
+	 */
+	bool OnProcessMessageReceived(CefRefPtr<CefBrowser> Browser, CefProcessId SourceProcess, CefRefPtr<CefProcessMessage> Message);
+
+	/**
+	 * Sends a message to the renderer process. 
+	 * See [REDACTED_PROFILE_URL]/cef/wiki/GeneralUsage#markdown-header-inter-process-communication-ipc for more information.
+	 *
+	 * @paran Message the message to send to the renderer process
+	 */
+	void SendProcessMessage(CefRefPtr<CefProcessMessage> Message);
+
+	CefRefPtr<CefDictionaryValue> ConvertStruct(UStruct* TypeInfo, void* StructPtr);
+	CefRefPtr<CefDictionaryValue> ConvertObject(UObject* Object);
+	CefRefPtr<CefDictionaryValue> ConvertObject(const FGuid& Key);
+
+	// Works for CefListValue and CefDictionaryValues
+	template<typename ContainerType, typename KeyType>
+	bool SetConverted(CefRefPtr<ContainerType> Container, KeyType Key, FWebJSParam& Param)
+	{
+		switch (Param.Tag)
+		{
+			case FWebJSParam::PTYPE_NULL:
+				return Container->SetNull(Key);
+			case FWebJSParam::PTYPE_BOOL:
+				return Container->SetBool(Key, Param.BoolValue);
+			case FWebJSParam::PTYPE_DOUBLE:
+				return Container->SetDouble(Key, Param.DoubleValue);
+			case FWebJSParam::PTYPE_INT:
+				return Container->SetInt(Key, Param.IntValue);
+			case FWebJSParam::PTYPE_STRING:
+			{
+				CefString ConvertedString = **Param.StringValue;
+				return Container->SetString(Key, ConvertedString);
+			}
+			case FWebJSParam::PTYPE_OBJECT:
+			{
+				CefRefPtr<CefDictionaryValue> ConvertedObject = ConvertObject(Param.ObjectValue);
+				return Container->SetDictionary(Key, ConvertedObject);
+			}
+			case FWebJSParam::PTYPE_STRUCT:
+			{
+				CefRefPtr<CefDictionaryValue> ConvertedStruct = ConvertStruct(Param.StructValue.TypeInfo, Param.StructValue.StructPtr);
+				return Container->SetDictionary(Key, ConvertedStruct);
+			}
+			default:
+				return false;
+		}
+	}
+
+	FWebJSParam ConvertResult(UProperty* ResultProperty, uint8* Data);
+	void BindCefBrowser(CefRefPtr<CefBrowser> Browser);
+	void UnbindCefBrowser();
+	CefRefPtr<CefDictionaryValue> GetPermanentBindings();
+
+	void InvokeJSFunction(FGuid FunctionId, int32 ArgCount, FWebJSParam Arguments[], bool bIsError=false);
+	void InvokeJSErrorResult(FGuid FunctionId, const FString& Error);
+
+public:
+
+	// FGCObject API
+	virtual void AddReferencedObjects( FReferenceCollector& Collector ) override;
+
+private:
+	bool ConvertStructArgImpl(uint8* Args, UProperty* Param, CefRefPtr<CefListValue> List, int32 Index);
+
+	bool IsValid()
+	{
+		return InternalCefBrowser.get() != nullptr;
+	}
+
+	// Creates a reversible memory addres -> psuedo-guid mapping.
+	// This is done by xoring the address with the first 64 bits of a base guid owned by the instance.
+	// Used to identify UObjects from the render process withough exposing internal pointers.
+	FGuid PtrToGuid(UObject* Ptr)
+	{
+		FGuid Guid = BaseGuid;
+		if (Ptr == nullptr)
+		{
+			Guid.Invalidate();
+		}
+		else
+		{
+			UPTRINT IntPtr = reinterpret_cast<UPTRINT>(Ptr);
+			if (sizeof(UPTRINT) > 4)
+			{
+				Guid[0] ^= IntPtr >> 32;
+			}
+			Guid[1] ^= IntPtr & 0xFFFFFFFF;
+		}
+		return Guid;
+	}
+
+	// In addition to reversing the mapping, it verifies that we are currently holding on to an instance of that UObject
+	UObject* GuidToPtr(const FGuid& Guid)
+	{
+		UPTRINT IntPtr = 0;
+		if (sizeof(UPTRINT) > 4)
+		{
+			IntPtr = (UPTRINT)(Guid[0] ^ BaseGuid[0]) << 32;
+		}
+		IntPtr |= (Guid[1] ^ BaseGuid[1]) & 0xFFFFFFFF;
+
+		UObject* Result = reinterpret_cast<UObject*>(IntPtr);
+		if (BoundObjects.Contains(Result))
+		{
+			return Result;
+		}
+		else
+		{
+			return nullptr;
+		}
+	}
+
+	void RetainBinding(UObject* Object)
+	{
+		if (BoundObjects.Contains(Object))
+		{
+			if(!BoundObjects[Object].bIsPermanent)
+			{
+				BoundObjects[Object].Refcount++;
+			}
+		}
+		else
+		{
+			BoundObjects.Add(Object, {false, 1});
+		}
+	}
+
+	void ReleaseBinding(UObject* Object)
+	{
+		if (BoundObjects.Contains(Object))
+		{
+			auto& Binding = BoundObjects[Object];
+			if(!Binding.bIsPermanent)
+			{
+				Binding.Refcount--;
+				if (Binding.Refcount <= 0)
+				{
+					BoundObjects.Remove(Object);
+				}
+			}
+		}
+	}
+
+	/** Message handling helpers */
+
+	bool HandleExecuteUObjectMethodMessage(CefRefPtr<CefListValue> MessageArguments);
+	bool HandleReleaseUObjectMessage(CefRefPtr<CefListValue> MessageArguments);
+
+private:
+
+	struct ObjectBinding
+	{
+		bool bIsPermanent;
+		int32 Refcount;
+	};
+
+	/** Private data */
+	FGuid BaseGuid;
+
+	/** Pointer to the CEF Browser for this window. */
+	CefRefPtr<CefBrowser> InternalCefBrowser;
+
+	/** UObjects currently visible on the renderer side. */
+	TMap<UObject*, ObjectBinding> BoundObjects;
+
+	/** Reverse lookup for permanent bindings */
+	TMap<FString, UObject*> PermanentUObjectsByName;
+
+};
+
+#endif
\ No newline at end of file
diff --git a/Engine/Source/Runtime/WebBrowser/Private/WebJSStructDeserializerBackend.cpp b/Engine/Source/Runtime/WebBrowser/Private/WebJSStructDeserializerBackend.cpp
new file mode 100644
--- /dev/null
+++ b/Engine/Source/Runtime/WebBrowser/Private/WebJSStructDeserializerBackend.cpp
@@ -0,0 +1,367 @@
+// Copyright 1998-2015 Epic Games, Inc. All Rights Reserved.
+
+#include "WebBrowserPrivatePCH.h"
+#include "WebJSStructDeserializerBackend.h"
+
+#if WITH_CEF3
+/* Internal helpers
+ *****************************************************************************/
+namespace {
+
+	template<typename ValueType, typename ContainerType, typename KeyType>
+	ValueType GetNumeric(CefRefPtr<ContainerType> Container, KeyType Key)
+	{
+		switch(Container->GetType(Key))
+		{
+			case VTYPE_BOOL:
+				return static_cast<ValueType>(Container->GetBool(Key));
+			case VTYPE_INT:
+				return static_cast<ValueType>(Container->GetInt(Key));
+			case VTYPE_DOUBLE:
+				return static_cast<ValueType>(Container->GetDouble(Key));
+			case VTYPE_STRING:
+			case VTYPE_DICTIONARY:
+			case VTYPE_LIST:
+			case VTYPE_NULL:
+			case VTYPE_BINARY:
+			default:
+				return static_cast<ValueType>(0);
+		}
+	}
+
+	template<typename ContainerType, typename KeyType>
+	void AssignTokenFromContainer(ContainerType Container, KeyType Key,  EStructDeserializerBackendTokens& OutToken, FString& PropertyName, TSharedPtr<ICefContainerWalker>& Retval)
+	{
+		switch (Container->GetType(Key))
+		{
+			case VTYPE_NULL:
+			case VTYPE_BOOL:
+			case VTYPE_INT:
+			case VTYPE_DOUBLE:
+			case VTYPE_STRING:
+				OutToken = EStructDeserializerBackendTokens::Property;
+				break;
+			case VTYPE_DICTIONARY:
+			{
+				CefRefPtr<CefDictionaryValue> Dictionary = Container->GetDictionary(Key);
+				if (Dictionary->GetType("$type") == VTYPE_STRING )
+				{
+					OutToken = EStructDeserializerBackendTokens::Property;
+				}
+				else
+				{
+					TSharedPtr<ICefContainerWalker> NewWalker(new FCefDictionaryValueWalker(Retval, Dictionary));
+					Retval = NewWalker->GetNextToken(OutToken, PropertyName);
+				}
+				break;
+			}
+			case VTYPE_LIST:
+			{
+				TSharedPtr<ICefContainerWalker> NewWalker(new FCefListValueWalker(Retval, Container->GetList(Key)));
+				Retval = NewWalker->GetNextToken(OutToken, PropertyName);
+				break;
+			}
+			case VTYPE_BINARY:
+			case VTYPE_INVALID:
+			default:
+				OutToken = EStructDeserializerBackendTokens::Error;
+				break;
+		}
+	}
+
+
+	/**
+	 * Sets the value of the given property.
+	 *
+	 * @param Property The property to set.
+	 * @param Outer The property that contains the property to be set, if any.
+	 * @param Data A pointer to the memory holding the property's data.
+	 * @param ArrayIndex The index of the element to set (if the property is an array).
+	 * @return true on success, false otherwise.
+	 * @see ClearPropertyValue
+	 */
+	template<typename UPropertyType, typename PropertyType>
+	bool SetPropertyValue( UProperty* Property, UProperty* Outer, void* Data, int32 ArrayIndex, const PropertyType& Value )
+	{
+		PropertyType* ValuePtr = nullptr;
+		UArrayProperty* ArrayProperty = Cast<UArrayProperty>(Outer);
+
+		if (ArrayProperty != nullptr)
+		{
+			if (ArrayProperty->Inner != Property)
+			{
+				return false;
+			}
+
+			FScriptArrayHelper ArrayHelper(ArrayProperty, ArrayProperty->template ContainerPtrToValuePtr<void>(Data));
+			int32 Index = ArrayHelper.AddValue();
+		
+			ValuePtr = (PropertyType*)ArrayHelper.GetRawPtr(Index);
+		}
+		else
+		{
+			UPropertyType* TypedProperty = Cast<UPropertyType>(Property);
+
+			if (TypedProperty == nullptr)
+			{
+				return false;
+			}
+
+			ValuePtr = TypedProperty->template ContainerPtrToValuePtr<PropertyType>(Data, ArrayIndex);
+		}
+
+		if (ValuePtr == nullptr)
+		{
+			return false;
+		}
+
+		*ValuePtr = Value;
+
+		return true;
+	}
+
+
+	template<typename PropertyType, typename ContainerType, typename KeyType>
+	bool ReadNumericProperty(UProperty* Property, UProperty* Outer, void* Data, int32 ArrayIndex, CefRefPtr<ContainerType> Container, KeyType Key )
+	{
+		typedef typename PropertyType::TCppType TCppType;
+		if (Property->IsA<PropertyType>())
+		{
+			return SetPropertyValue<PropertyType, TCppType>(Property, Outer, Data, ArrayIndex, GetNumeric<TCppType>(Container, Key));
+		}
+		else
+		{
+			return false;
+		}
+	}
+
+	template<typename ContainerType, typename KeyType>
+	bool ReadBoolProperty(UProperty* Property, UProperty* Outer, void* Data, int32 ArrayIndex, CefRefPtr<ContainerType> Container, KeyType Key )
+	{
+		if (Property->IsA<UBoolProperty>())
+		{
+			return SetPropertyValue<UBoolProperty, bool>(Property, Outer, Data, ArrayIndex, GetNumeric<int>(Container, Key)!=0);
+		}
+		return false;
+		
+	}
+
+	template<typename ContainerType, typename KeyType>
+	bool ReadJSFunctionProperty(TSharedPtr<FWebJSScripting> Scripting, UProperty* Property, UProperty* Outer, void* Data, int32 ArrayIndex, CefRefPtr<ContainerType> Container, KeyType Key )
+	{
+		if (Container->GetType(Key) != VTYPE_DICTIONARY || !Property->IsA<UStructProperty>())
+		{
+			return false;
+		}
+		CefRefPtr<CefDictionaryValue> Dictionary = Container->GetDictionary(Key);
+		UStructProperty* StructProperty = Cast<UStructProperty>(Property);
+
+		if ( StructProperty->Struct != FWebJSFunction::StaticStruct())
+		{
+			return false;
+		}
+
+		FGuid CallbackID;
+		if (!FGuid::Parse(FString(Dictionary->GetString("$id").ToWString().c_str()), CallbackID))
+		{
+			// Invalid GUID
+			return false;
+		}
+		FWebJSFunction CallbackObject(Scripting, CallbackID);
+		return SetPropertyValue<UStructProperty, FWebJSFunction>(Property, Outer, Data, ArrayIndex, CallbackObject);
+	}
+
+	template<typename ContainerType, typename KeyType>
+	bool ReadStringProperty(UProperty* Property, UProperty* Outer, void* Data, int32 ArrayIndex, CefRefPtr<ContainerType> Container, KeyType Key )
+	{
+		if (Container->GetType(Key) == VTYPE_STRING)
+		{
+			FString StringValue = Container->GetString(Key).ToWString().c_str();
+
+			if (Property->IsA<UStrProperty>())
+			{
+				return SetPropertyValue<UStrProperty, FString>(Property, Outer, Data, ArrayIndex, StringValue);
+			}
+			else if (Property->IsA<UNameProperty>())
+			{
+				return SetPropertyValue<UNameProperty, FName>(Property, Outer, Data, ArrayIndex, *StringValue);
+			}
+			else if (Property->IsA<UTextProperty>())
+			{
+				return SetPropertyValue<UTextProperty, FText>(Property, Outer, Data, ArrayIndex, FText::FromString(StringValue));
+			}
+			else if (Property->IsA<UByteProperty>())
+			{
+				UByteProperty* ByteProperty = Cast<UByteProperty>(Property);
+				int32 Index = ByteProperty->Enum->FindEnumIndex(*StringValue);
+
+				if (Index == INDEX_NONE)
+				{
+					return false;
+				}
+
+				return SetPropertyValue<UByteProperty, uint8>(Property, Outer, Data, ArrayIndex, (uint8)Index);
+			}
+		}
+
+		return false;
+	}
+
+	template<typename ContainerType, typename KeyType>
+	bool ReadProperty(TSharedPtr<FWebJSScripting> Scripting, UProperty* Property, UProperty* Outer, void* Data, int32 ArrayIndex, CefRefPtr<ContainerType> Container, KeyType Key )
+	{
+		return ReadBoolProperty(Property, Outer, Data, ArrayIndex, Container, Key)
+			|| ReadStringProperty(Property, Outer, Data, ArrayIndex, Container, Key)
+			|| ReadNumericProperty<UByteProperty>(Property, Outer, Data, ArrayIndex, Container, Key)
+			|| ReadNumericProperty<UInt8Property>(Property, Outer, Data, ArrayIndex, Container, Key)
+			|| ReadNumericProperty<UInt16Property>(Property, Outer, Data, ArrayIndex, Container, Key)
+			|| ReadNumericProperty<UIntProperty>(Property, Outer, Data, ArrayIndex, Container, Key)
+			|| ReadNumericProperty<UInt64Property>(Property, Outer, Data, ArrayIndex, Container, Key)
+			|| ReadNumericProperty<UUInt16Property>(Property, Outer, Data, ArrayIndex, Container, Key)
+			|| ReadNumericProperty<UUInt32Property>(Property, Outer, Data, ArrayIndex, Container, Key)
+			|| ReadNumericProperty<UUInt64Property>(Property, Outer, Data, ArrayIndex, Container, Key)
+			|| ReadNumericProperty<UFloatProperty>(Property, Outer, Data, ArrayIndex, Container, Key)
+			|| ReadNumericProperty<UDoubleProperty>(Property, Outer, Data, ArrayIndex, Container, Key)
+			|| ReadJSFunctionProperty(Scripting, Property, Outer, Data, ArrayIndex, Container, Key);
+	}
+
+
+}
+
+TSharedPtr<ICefContainerWalker> FCefListValueWalker::GetNextToken(EStructDeserializerBackendTokens& OutToken, FString& PropertyName)
+{
+	TSharedPtr<ICefContainerWalker> Retval = SharedThis(this);
+	Index++;
+	if (Index == -1)
+	{
+		OutToken = EStructDeserializerBackendTokens::ArrayStart;
+	}
+	else if ( Index < List->GetSize() )
+	{
+		AssignTokenFromContainer(List, Index, OutToken, PropertyName, Retval);
+		PropertyName = FString();
+	}
+	else
+	{
+		OutToken = EStructDeserializerBackendTokens::ArrayEnd;
+		Retval = Parent;
+	}
+	return Retval;
+}
+
+bool FCefListValueWalker::ReadProperty(TSharedPtr<FWebJSScripting> Scripting, UProperty* Property, UProperty* Outer, void* Data, int32 ArrayIndex)
+{
+	return ::ReadProperty(Scripting, Property, Outer, Data, ArrayIndex, List, Index);
+}
+
+TSharedPtr<ICefContainerWalker> FCefDictionaryValueWalker::GetNextToken(EStructDeserializerBackendTokens& OutToken, FString& PropertyName)
+{
+	TSharedPtr<ICefContainerWalker> Retval = SharedThis(this);
+	Index++;
+	if (Index == -1)
+	{
+		OutToken = EStructDeserializerBackendTokens::StructureStart;
+	}
+	else if ( Index < Keys.size() )
+	{
+		AssignTokenFromContainer(Dictionary, Keys[Index], OutToken, PropertyName, Retval);
+		PropertyName = Keys[Index].ToWString().c_str();
+	}
+	else
+	{
+		OutToken = EStructDeserializerBackendTokens::StructureEnd;
+		Retval = Parent;
+	}
+	return Retval;
+}
+
+bool FCefDictionaryValueWalker::ReadProperty(TSharedPtr<FWebJSScripting> Scripting, UProperty* Property, UProperty* Outer, void* Data, int32 ArrayIndex)
+{
+	return ::ReadProperty(Scripting, Property, Outer, Data, ArrayIndex, Dictionary, Keys[Index]);
+}
+
+
+/* IStructDeserializerBackend interface
+ *****************************************************************************/
+
+
+
+const FString& FWebJSStructDeserializerBackend::GetCurrentPropertyName() const
+{
+	return CurrentPropertyName;
+}
+
+
+FString FWebJSStructDeserializerBackend::GetDebugString() const
+{
+	return CurrentPropertyName;
+}
+
+
+const FString& FWebJSStructDeserializerBackend::GetLastErrorMessage() const
+{
+	return CurrentPropertyName;
+}
+
+
+bool FWebJSStructDeserializerBackend::GetNextToken( EStructDeserializerBackendTokens& OutToken )
+{
+	if (Walker.IsValid())
+	{
+		Walker = Walker->GetNextToken(OutToken, CurrentPropertyName);
+		return true;
+	}
+	else
+	{
+		return false;
+	}
+}
+
+
+bool FWebJSStructDeserializerBackend::ReadProperty( UProperty* Property, UProperty* Outer, void* Data, int32 ArrayIndex )
+{
+	return Walker->ReadProperty(Scripting, Property, Outer, Data, ArrayIndex);
+}
+
+
+void FWebJSStructDeserializerBackend::SkipArray() 
+{
+	EStructDeserializerBackendTokens Token;
+	int32 depth = 1;
+	while (GetNextToken(Token) && depth > 0)
+	{
+		switch (Token)
+		{
+		case EStructDeserializerBackendTokens::ArrayEnd:
+			depth --;
+			break;
+		case EStructDeserializerBackendTokens::ArrayStart:
+			depth ++;
+			break;
+		default:
+			break;
+		}
+	}
+}
+
+void FWebJSStructDeserializerBackend::SkipStructure()
+{
+	EStructDeserializerBackendTokens Token;
+	int32 depth = 1;
+	while (GetNextToken(Token) && depth > 0)
+	{
+		switch (Token)
+		{
+		case EStructDeserializerBackendTokens::StructureEnd:
+			depth --;
+			break;
+		case EStructDeserializerBackendTokens::StructureStart:
+			depth ++;
+			break;
+		default:
+			break;
+		}
+	}
+}
+
+#endif
\ No newline at end of file
diff --git a/Engine/Source/Runtime/WebBrowser/Private/WebJSStructDeserializerBackend.h b/Engine/Source/Runtime/WebBrowser/Private/WebJSStructDeserializerBackend.h
new file mode 100644
--- /dev/null
+++ b/Engine/Source/Runtime/WebBrowser/Private/WebJSStructDeserializerBackend.h
@@ -0,0 +1,117 @@
+// Copyright 1998-2015 Epic Games, Inc. All Rights Reserved.
+
+#pragma once
+
+#include "IStructDeserializerBackend.h"
+#include "Core.h"
+#include "WebJSScripting.h"
+#if WITH_CEF3
+
+// forward declarations
+class UProperty;
+class UStruct;
+
+#if PLATFORM_WINDOWS
+#include "AllowWindowsPlatformTypes.h"
+#endif
+#pragma push_macro("OVERRIDE")
+#undef OVERRIDE // cef headers provide their own OVERRIDE macro
+#include "include/cef_values.h"
+#pragma pop_macro("OVERRIDE")
+#if PLATFORM_WINDOWS
+#include "HideWindowsPlatformTypes.h"
+#endif
+
+
+class ICefContainerWalker
+	: public TSharedFromThis<ICefContainerWalker>
+{
+public:
+	ICefContainerWalker(TSharedPtr<ICefContainerWalker> InParent)
+		: Parent(InParent)
+	{}
+	virtual ~ICefContainerWalker() {}
+
+	virtual TSharedPtr<ICefContainerWalker> GetNextToken( EStructDeserializerBackendTokens& OutToken, FString& PropertyName ) = 0;
+	virtual bool ReadProperty(TSharedPtr<FWebJSScripting> Scripting, UProperty* Property, UProperty* Outer, void* Data, int32 ArrayIndex) = 0;
+
+	TSharedPtr<ICefContainerWalker> Parent;
+};
+
+class FCefListValueWalker
+	: public ICefContainerWalker
+{
+public:
+	FCefListValueWalker(TSharedPtr<ICefContainerWalker> InParent, CefRefPtr<CefListValue> InList)
+		: ICefContainerWalker(InParent)
+		, List(InList)
+		, Index(-2)
+	{}
+
+	virtual TSharedPtr<ICefContainerWalker> GetNextToken( EStructDeserializerBackendTokens& OutToken, FString& PropertyName ) override;
+	virtual bool ReadProperty(TSharedPtr<FWebJSScripting> Scripting, UProperty* Property, UProperty* Outer, void* Data, int32 ArrayIndex ) override;
+
+	CefRefPtr<CefListValue> List;
+	int32 Index;
+};
+
+class FCefDictionaryValueWalker
+	: public ICefContainerWalker
+{
+public:
+	FCefDictionaryValueWalker(TSharedPtr<ICefContainerWalker> InParent, CefRefPtr<CefDictionaryValue> InDictionary)
+		: ICefContainerWalker(InParent)
+		, Dictionary(InDictionary)
+		, Index(-2)
+	{
+		Dictionary->GetKeys(Keys);
+	}
+
+	virtual TSharedPtr<ICefContainerWalker> GetNextToken( EStructDeserializerBackendTokens& OutToken, FString& PropertyName ) override;
+	virtual bool ReadProperty(TSharedPtr<FWebJSScripting> Scripting, UProperty* Property, UProperty* Outer, void* Data, int32 ArrayIndex ) override;
+
+private:
+	CefRefPtr<CefDictionaryValue> Dictionary;
+	int32 Index;
+	CefDictionaryValue::KeyList Keys;
+};
+
+/**
+ * Implements a writer for UStruct serialization using CefDictionary.
+ */
+class FWebJSStructDeserializerBackend
+	: public IStructDeserializerBackend
+{
+public:
+
+	/**
+	 * Creates and initializes a new instance.
+	 *
+	 * @param Archive The archive to deserialize from.
+	 */
+	FWebJSStructDeserializerBackend(TSharedPtr<FWebJSScripting> InScripting, CefRefPtr<CefDictionaryValue> InDictionary)
+		: Scripting(InScripting)
+		, Walker(new FCefDictionaryValueWalker(nullptr, InDictionary))
+		, CurrentPropertyName()
+	{ }
+
+public:
+
+	// IStructDeserializerBackend interface
+
+	virtual const FString& GetCurrentPropertyName() const override;
+	virtual FString GetDebugString() const override;
+	virtual const FString& GetLastErrorMessage() const override;
+	virtual bool GetNextToken( EStructDeserializerBackendTokens& OutToken ) override;
+	virtual bool ReadProperty( UProperty* Property, UProperty* Outer, void* Data, int32 ArrayIndex ) override;
+	virtual void SkipArray() override;
+	virtual void SkipStructure() override;
+
+private:
+	TSharedPtr<FWebJSScripting> Scripting;
+	/** Holds the source CEF dictionary containing a serialized verion of the structure. */
+	TSharedPtr<ICefContainerWalker> Walker;
+	FString CurrentPropertyName;
+};
+
+#endif
\ No newline at end of file
diff --git a/Engine/Source/Runtime/WebBrowser/Private/WebJSStructSerializerBackend.cpp b/Engine/Source/Runtime/WebBrowser/Private/WebJSStructSerializerBackend.cpp
new file mode 100644
--- /dev/null
+++ b/Engine/Source/Runtime/WebBrowser/Private/WebJSStructSerializerBackend.cpp
@@ -0,0 +1,256 @@
+// Copyright 1998-2015 Epic Games, Inc. All Rights Reserved.
+
+#include "WebBrowserPrivatePCH.h"
+#include "WebJSStructSerializerBackend.h"
+
+#if WITH_CEF3
+
+/* Private methods
+ *****************************************************************************/
+
+void FWebJSStructSerializerBackend::AddNull(UProperty* Property)
+{
+	StackItem& Current = Stack.Top();
+	switch (Current.Kind) {
+		case StackItem::STYPE_DICTIONARY:
+			Current.DictionaryValue->SetNull(*Property->GetName());
+			break;
+		case StackItem::STYPE_LIST:
+			Current.ListValue->SetNull(Current.ListValue->GetSize());
+		break;
+	}
+}
+
+void FWebJSStructSerializerBackend::Add(UProperty* Property, bool Value)
+{
+	StackItem& Current = Stack.Top();
+	switch (Current.Kind) {
+		case StackItem::STYPE_DICTIONARY:
+			Current.DictionaryValue->SetBool(*Property->GetName(), Value);
+			break;
+		case StackItem::STYPE_LIST:
+			Current.ListValue->SetBool(Current.ListValue->GetSize(), Value);
+		break;
+	}
+}
+
+void FWebJSStructSerializerBackend::Add(UProperty* Property, int32 Value)
+{
+	StackItem& Current = Stack.Top();
+	switch (Current.Kind) {
+		case StackItem::STYPE_DICTIONARY:
+			Current.DictionaryValue->SetInt(*Property->GetName(), Value);
+			break;
+		case StackItem::STYPE_LIST:
+			Current.ListValue->SetInt(Current.ListValue->GetSize(), Value);
+		break;
+	}
+}
+
+void FWebJSStructSerializerBackend::Add(UProperty* Property, double Value)
+{
+	StackItem& Current = Stack.Top();
+	switch (Current.Kind) {
+		case StackItem::STYPE_DICTIONARY:
+			Current.DictionaryValue->SetDouble(*Property->GetName(), Value);
+			break;
+		case StackItem::STYPE_LIST:
+			Current.ListValue->SetDouble(Current.ListValue->GetSize(), Value);
+		break;
+	}
+}
+
+void FWebJSStructSerializerBackend::Add(UProperty* Property, FString Value)
+{
+	StackItem& Current = Stack.Top();
+	switch (Current.Kind) {
+		case StackItem::STYPE_DICTIONARY:
+			Current.DictionaryValue->SetString(*Property->GetName(), *Value);
+			break;
+		case StackItem::STYPE_LIST:
+			Current.ListValue->SetString(Current.ListValue->GetSize(), *Value);
+		break;
+	}
+}
+
+void FWebJSStructSerializerBackend::Add(UProperty* Property, UObject* Value)
+{
+	StackItem& Current = Stack.Top();
+	switch (Current.Kind) {
+		case StackItem::STYPE_DICTIONARY:
+			Current.DictionaryValue->SetDictionary(*Property->GetName(), Scripting->ConvertObject(Value));
+			break;
+		case StackItem::STYPE_LIST:
+			Current.ListValue->SetDictionary(Current.ListValue->GetSize(), Scripting->ConvertObject(Value));
+		break;
+	}
+}
+
+/* IStructSerializerBackend interface
+ *****************************************************************************/
+
+void FWebJSStructSerializerBackend::BeginArray( UProperty* Property )
+{
+	CefRefPtr<CefListValue> ListValue = CefListValue::Create();
+	Stack.Push(StackItem(Property->GetName(), ListValue));
+}
+
+void FWebJSStructSerializerBackend::BeginStructure( UProperty* Property )
+{
+	CefRefPtr<CefDictionaryValue> DictionaryValue = CefDictionaryValue::Create();
+	Stack.Push(StackItem(Property->GetName(), DictionaryValue));
+}
+
+void FWebJSStructSerializerBackend::BeginStructure( UStruct* TypeInfo )
+{
+	Result = CefDictionaryValue::Create();
+	Stack.Push(StackItem(FString(), Result));
+}
+
+
+void FWebJSStructSerializerBackend::EndArray( UProperty* Property )
+{
+	StackItem Previous = Stack.Pop();
+	check(Previous.Kind == StackItem::STYPE_LIST);
+	check(Stack.Num() > 0); // The root level object is always a struct
+	StackItem& Current = Stack.Top();
+
+	switch (Current.Kind) {
+		case StackItem::STYPE_DICTIONARY:
+			Current.DictionaryValue->SetList(*Previous.Name, Previous.ListValue);
+			break;
+		case StackItem::STYPE_LIST:
+			Current.ListValue->SetList(Current.ListValue->GetSize(), Previous.ListValue);
+		break;
+	}
+}
+
+
+void FWebJSStructSerializerBackend::EndStructure()
+{
+	StackItem Previous = Stack.Pop();
+	check(Previous.Kind == StackItem::STYPE_DICTIONARY);
+
+	if (Stack.Num() > 0)
+	{
+		StackItem& Current = Stack.Top();
+
+		switch (Current.Kind) {
+			case StackItem::STYPE_DICTIONARY:
+				Current.DictionaryValue->SetDictionary(*Previous.Name, Previous.DictionaryValue);
+				break;
+			case StackItem::STYPE_LIST:
+				Current.ListValue->SetDictionary(Current.ListValue->GetSize(), Previous.DictionaryValue);
+			break;
+		}
+	}
+	else
+	{
+		check(Result == Previous.DictionaryValue);
+	}
+}
+
+
+void FWebJSStructSerializerBackend::WriteComment( const FString& Comment )
+{
+	// Cef values do not support comments
+}
+
+
+void FWebJSStructSerializerBackend::WriteProperty( UProperty* Property, const void* Data, UStruct* TypeInfo, int32 ArrayIndex )
+{
+
+	// booleans
+	if (TypeInfo == UBoolProperty::StaticClass())
+	{
+		Add(Property, Cast<UBoolProperty>(Property)->GetPropertyValue_InContainer(Data, ArrayIndex));
+	}
+
+	// unsigned bytes & enumerations
+	else if (TypeInfo == UByteProperty::StaticClass())
+	{
+		UByteProperty* ByteProperty = Cast<UByteProperty>(Property);
+
+		if (ByteProperty->IsEnum())
+		{
+			Add(Property, ByteProperty->Enum->GetEnumName(ByteProperty->GetPropertyValue_InContainer(Data, ArrayIndex)));
+		}
+		else
+		{
+			Add(Property, (double)Cast<UByteProperty>(Property)->GetPropertyValue_InContainer(Data, ArrayIndex));
+		}			
+	}
+
+	// floating point numbers
+	else if (TypeInfo == UDoubleProperty::StaticClass())
+	{
+		Add(Property, Cast<UDoubleProperty>(Property)->GetPropertyValue_InContainer(Data, ArrayIndex));
+	}
+	else if (TypeInfo == UFloatProperty::StaticClass())
+	{
+		Add(Property, Cast<UFloatProperty>(Property)->GetPropertyValue_InContainer(Data, ArrayIndex));
+	}
+
+	// signed integers
+	else if (TypeInfo == UIntProperty::StaticClass())
+	{
+		Add(Property, (int32)Cast<UIntProperty>(Property)->GetPropertyValue_InContainer(Data, ArrayIndex));
+	}
+	else if (TypeInfo == UInt8Property::StaticClass())
+	{
+		Add(Property, (int32)Cast<UInt8Property>(Property)->GetPropertyValue_InContainer(Data, ArrayIndex));
+	}
+	else if (TypeInfo == UInt16Property::StaticClass())
+	{
+		Add(Property, (int32)Cast<UInt16Property>(Property)->GetPropertyValue_InContainer(Data, ArrayIndex));
+	}
+	else if (TypeInfo == UInt64Property::StaticClass())
+	{
+		Add(Property, (double)Cast<UInt64Property>(Property)->GetPropertyValue_InContainer(Data, ArrayIndex));
+	}
+
+	// unsigned integers
+	else if (TypeInfo == UUInt16Property::StaticClass())
+	{
+		Add(Property, (int32)Cast<UUInt16Property>(Property)->GetPropertyValue_InContainer(Data, ArrayIndex));
+	}
+	else if (TypeInfo == UUInt32Property::StaticClass())
+	{
+		Add(Property, (double)Cast<UUInt32Property>(Property)->GetPropertyValue_InContainer(Data, ArrayIndex));
+	}
+	else if (TypeInfo == UUInt64Property::StaticClass())
+	{
+		Add(Property, (double)Cast<UUInt64Property>(Property)->GetPropertyValue_InContainer(Data, ArrayIndex));
+	}
+
+	// names & strings
+	else if (TypeInfo == UNameProperty::StaticClass())
+	{
+		Add(Property, Cast<UNameProperty>(Property)->GetPropertyValue_InContainer(Data, ArrayIndex).ToString());
+	}
+	else if (TypeInfo == UStrProperty::StaticClass())
+	{
+		Add(Property, Cast<UStrProperty>(Property)->GetPropertyValue_InContainer(Data, ArrayIndex));
+	}
+	else if (TypeInfo == UTextProperty::StaticClass())
+	{
+		Add(Property, Cast<UTextProperty>(Property)->GetPropertyValue_InContainer(Data, ArrayIndex).ToString());
+	}
+
+	// classes & objects
+	else if (TypeInfo == UClassProperty::StaticClass())
+	{
+		Add(Property, Cast<UClassProperty>(Property)->GetPropertyValue_InContainer(Data, ArrayIndex)->GetPathName());
+	}
+	else if (TypeInfo == UObjectProperty::StaticClass())
+	{
+		Add(Property, Cast<UObjectProperty>(Property)->GetPropertyValue_InContainer(Data, ArrayIndex));
+	}
+	
+	else
+	{
+		GLog->Logf(ELogVerbosity::Warning, TEXT("FWebJSStructSerializerBackend: Property %s cannot be serialized, because its type (%s) is not supported"), *Property->GetName(), *TypeInfo->GetName());
+	}
+}
+
+#endif
diff --git a/Engine/Source/Runtime/WebBrowser/Private/WebJSStructSerializerBackend.h b/Engine/Source/Runtime/WebBrowser/Private/WebJSStructSerializerBackend.h
new file mode 100644
--- /dev/null
+++ b/Engine/Source/Runtime/WebBrowser/Private/WebJSStructSerializerBackend.h
@@ -0,0 +1,94 @@
+// Copyright 1998-2015 Epic Games, Inc. All Rights Reserved.
+
+#pragma once
+
+#include "IStructSerializerBackend.h"
+#include "Core.h"
+#include "WebJSScripting.h"
+#if WITH_CEF3
+
+// forward declarations
+class UProperty;
+class UStruct;
+
+#if PLATFORM_WINDOWS
+#include "AllowWindowsPlatformTypes.h"
+#endif
+#pragma push_macro("OVERRIDE")
+#undef OVERRIDE // cef headers provide their own OVERRIDE macro
+#include "include/cef_values.h"
+#pragma pop_macro("OVERRIDE")
+#if PLATFORM_WINDOWS
+#include "HideWindowsPlatformTypes.h"
+#endif
+
+/**
+ * Implements a writer for UStruct serialization using CefDictionary.
+ */
+class FWebJSStructSerializerBackend
+	: public IStructSerializerBackend
+{
+public:
+
+	/**
+	 * Creates and initializes a new instance.
+	 *	 */
+	FWebJSStructSerializerBackend(TSharedPtr<FWebJSScripting> InScripting)
+		: Scripting(InScripting), Stack(), Result()
+	{
+	}
+
+	CefRefPtr<CefDictionaryValue> GetResult()
+	{
+		return Result;
+	}
+
+public:
+
+	// IStructSerializerBackend interface
+
+	virtual void BeginArray( UProperty* Property ) override;
+	virtual void BeginStructure( UProperty* Property ) override;
+	virtual void BeginStructure( UStruct* TypeInfo ) override;
+	virtual void EndArray( UProperty* Property ) override;
+	virtual void EndStructure() override;
+	virtual void WriteComment( const FString& Comment ) override;
+	virtual void WriteProperty( UProperty* Property, const void* Data, UStruct* TypeInfo, int32 ArrayIndex ) override;
+
+private:
+
+	struct StackItem
+	{
+		enum {STYPE_DICTIONARY, STYPE_LIST} Kind;
+		FString Name;
+		CefRefPtr<CefDictionaryValue> DictionaryValue;
+		CefRefPtr<CefListValue> ListValue;
+
+		StackItem(const FString& InName, CefRefPtr<CefDictionaryValue> InDictionary)
+			: Kind(STYPE_DICTIONARY)
+			, Name(InName)
+			, DictionaryValue(InDictionary)
+			, ListValue()
+		{}
+
+		StackItem(const FString& InName, CefRefPtr<CefListValue> InList)
+			: Kind(STYPE_LIST)
+			, Name(InName)
+			, DictionaryValue()
+			, ListValue(InList)
+		{}
+
+	};
+	
+	TSharedPtr<FWebJSScripting> Scripting;
+	TArray<StackItem> Stack;
+	CefRefPtr<CefDictionaryValue> Result;
+	void AddNull(UProperty* Property);
+	void Add(UProperty* Property, bool Value);
+	void Add(UProperty* Property, int32 Value);
+	void Add(UProperty* Property, double Value);
+	void Add(UProperty* Property, FString Value);
+	void Add(UProperty* Property, UObject* Value);
+
+};
+#endif
\ No newline at end of file
diff --git a/Engine/Source/Runtime/WebBrowser/Public/IWebBrowserWindow.h b/Engine/Source/Runtime/WebBrowser/Public/IWebBrowserWindow.h
--- a/Engine/Source/Runtime/WebBrowser/Public/IWebBrowserWindow.h
+++ b/Engine/Source/Runtime/WebBrowser/Public/IWebBrowserWindow.h
@@ -207,6 +207,26 @@ public:
 	/** Close this window so that it can no longer be used. */
 	virtual void CloseBrowser() = 0;
 
+	/** 
+	 * Expose a UObject instance to the browser runtime.
+	 * Properties and Functions will be accessible from JavaScript side.
+	 * As all communication with the rendering procesis asynchronous, return values (both for properties and function results) are wrapped into JS Future objects.
+	 *
+	 * @param Name The name of the object. The object will show up as window.ue4.{Name} on the javascript side. If there is an existing object of the same name, this object will replace it. If bIsPermanent is false and there is an existing permanent binding, the permanent binding will be restored when the temporary one is removed.
+	 * @param Object The object instance.
+	 * @param bIsPermanent If true, the object will be visible to all pages loaded through this browser widget, otherwise, it will be deleted when navigating away from the current page. Non-permanent bindings should be registered from inside an OnLoadStarted event handler in order to be available before JS code starts loading.
+	 */
+	virtual void BindUObject(const FString& Name, UObject* Object, bool bIsPermanent = true) = 0;
+
+	/**
+	 * Remove an existing script binding registered by BindUObject.
+	 *
+	 * @param Name The name of the object to remove.
+	 * @param Object The object will only be removed if it is the same object as the one passed in.
+	 * @param bIsPermanent Must match the bIsPermanent argument passed to BindUObject.
+	 */
+	virtual void UnbindUObject(const FString& Name, UObject* Object, bool bIsPermanent = true) = 0;
+
 public:
 
 	/** A delegate that is invoked when the loading state of a document changed. */
diff --git a/Engine/Source/Runtime/WebBrowser/Public/SWebBrowser.h b/Engine/Source/Runtime/WebBrowser/Public/SWebBrowser.h
--- a/Engine/Source/Runtime/WebBrowser/Public/SWebBrowser.h
+++ b/Engine/Source/Runtime/WebBrowser/Public/SWebBrowser.h
@@ -4,10 +4,10 @@
 
 #include "SlateBasics.h"
 
-
 enum class EWebBrowserDocumentState;
 class IWebBrowserWindow;
 class FWebBrowserViewport;
+class UObject;
 
 DECLARE_DELEGATE_TwoParams(FJSQueryResultDelegate, int, FString);
 DECLARE_DELEGATE_RetVal_FourParams(bool, FOnJSQueryReceivedDelegate, int64, FString, bool, FJSQueryResultDelegate);
@@ -155,6 +155,26 @@ public:
 	/** Execute javascript on the current window */
 	void ExecuteJavascript(const FString& ScriptText);
 
+	/** 
+	 * Expose a UObject instance to the browser runtime.
+	 * Properties and Functions will be accessible from JavaScript side.
+	 * As all communication with the rendering procesis asynchronous, return values (both for properties and function results) are wrapped into JS Future objects.
+	 *
+	 * @param Name The name of the object. The object will show up as window.ue4.{Name} on the javascript side. If there is an existing object of the same name, this object will replace it. If bIsPermanent is false and there is an existing permanent binding, the permanent binding will be restored when the temporary one is removed.
+	 * @param Object The object instance.
+	 * @param bIsPermanent If true, the object will be visible to all pages loaded through this browser widget, otherwise, it will be deleted when navigating away from the current page. Non-permanent bindings should be registered from inside an OnLoadStarted event handler in order to be available before JS code starts loading.
+	 */
+	void BindUObject(const FString& Name, UObject* Object, bool bIsPermanent = true);
+
+	/**
+	 * Remove an existing script binding registered by BindUObject.
+	 *
+	 * @param Name The name of the object to remove.
+	 * @param Object The object will only be removed if it is the same object as the one passed in.
+	 * @param bIsPermanent Must match the bIsPermanent argument passed to BindUObject.
+	 */
+	void UnbindUObject(const FString& Name, UObject* Object, bool bIsPermanent = true);
+
 private:
 
 	/** Returns true if the browser can navigate backwards. */
diff --git a/Engine/Source/Runtime/WebBrowser/Public/WebJSFunction.h b/Engine/Source/Runtime/WebBrowser/Public/WebJSFunction.h
new file mode 100644
--- /dev/null
+++ b/Engine/Source/Runtime/WebBrowser/Public/WebJSFunction.h
@@ -0,0 +1,128 @@
+// Copyright 1998-2015 Epic Games, Inc. All Rights Reserved.
+
+#pragma once
+#include "CoreUObject.h"
+#include "WebJSFunction.generated.h"
+
+class FWebJSScripting;
+
+struct WEBBROWSER_API FWebJSParam
+{
+	enum { PTYPE_NULL, PTYPE_BOOL, PTYPE_INT, PTYPE_DOUBLE, PTYPE_STRING, PTYPE_OBJECT, PTYPE_STRUCT, } Tag;
+	union
+	{
+		bool BoolValue;
+		double DoubleValue;
+		int32 IntValue;
+		UObject* ObjectValue;
+		FString* StringValue;
+		struct {
+			UStruct* TypeInfo;
+			void* StructPtr;
+		} StructValue;
+	};
+
+	FWebJSParam() : Tag(PTYPE_NULL) {}
+	FWebJSParam(bool Value) : Tag(PTYPE_BOOL), BoolValue(Value) {}
+	FWebJSParam(int32 Value) : Tag(PTYPE_INT), IntValue(Value) {}
+	FWebJSParam(double Value) : Tag(PTYPE_DOUBLE), DoubleValue(Value) {}
+	FWebJSParam(const FString& Value) : Tag(PTYPE_STRING), StringValue(new FString(Value)) {}
+	FWebJSParam(const TCHAR* Value) : Tag(PTYPE_STRING), StringValue(new FString(Value)) {}
+	FWebJSParam(UObject* Value) : Tag(PTYPE_OBJECT), ObjectValue(Value) {}
+	FWebJSParam(UStruct* Struct, void* Value) : Tag(PTYPE_STRUCT), StructValue({Struct, Value}) {}
+	FWebJSParam(const FWebJSParam& Other);
+	~FWebJSParam();
+
+};
+
+/** Add this macro to a USTRUCT definition to be able to pass instances to FWebJSFunctions */
+#define DECLARE_JSPARAM_STRUCT() \
+	explicit operator FWebJSParam() \
+	{ \
+		return FWebJSParam(StaticStruct(), (void*)this);\
+	}
+
+USTRUCT()
+struct WEBBROWSER_API FWebJSFunction
+{
+	GENERATED_USTRUCT_BODY()
+
+	FWebJSFunction()
+		: ScriptingPtr()
+		, FunctionId()
+	{}
+
+	FWebJSFunction(TSharedPtr<FWebJSScripting> InScripting, const FGuid& InFunctionId)
+		: ScriptingPtr(InScripting)
+		, FunctionId(InFunctionId)
+	{}
+
+#if PLATFORM_COMPILER_HAS_VARIADIC_TEMPLATES
+	template<typename ...ArgTypes> void operator()(ArgTypes... Args) const
+	{
+		FWebJSParam ArgArray[sizeof...(Args)] = {FWebJSParam(Args)...};
+		Invoke(sizeof...(Args), ArgArray);
+	}
+#else
+	void operator()() const
+	{
+		Invoke(0, nullptr);
+	}
+
+	template<typename T1>
+	void operator()(T1 Arg1) const
+	{
+		FWebJSParam ArgArray[1] = {FWebJSParam(Arg1)};
+		Invoke(1, ArgArray);
+	}
+
+	template<typename T1, typename T2>
+	void operator()(T1 Arg1, T2 Arg2) const
+	{
+		FWebJSParam ArgArray[2] = {FWebJSParam(Arg1), FWebJSParam(Arg2)};
+		Invoke(2, ArgArray);
+	}
+
+	template<typename T1, typename T2, typename T3>
+	void operator()(T1 Arg1, T2 Arg2, T3 Arg3) const
+	{
+		FWebJSParam ArgArray[3] = {FWebJSParam(Arg1), FWebJSParam(Arg2), FWebJSParam(Arg3)};
+		Invoke(3, ArgArray);
+	}
+
+	template<typename T1, typename T2, typename T3, typename T4>
+	void operator()(T1 Arg1, T2 Arg2, T3 Arg3, T4 Arg4) const
+	{
+		FWebJSParam ArgArray[4] = {FWebJSParam(Arg1), FWebJSParam(Arg2), FWebJSParam(Arg3), FWebJSParam(Arg4)};
+		Invoke(4, ArgArray);
+	}
+
+	template<typename T1, typename T2, typename T3, typename T4, typename T5>
+	void operator()(T1 Arg1, T2 Arg2, T3 Arg3, T4 Arg4, T5 Arg5) const
+	{
+		FWebJSParam ArgArray[5] = {FWebJSParam(Arg1), FWebJSParam(Arg2), FWebJSParam(Arg3), FWebJSParam(Arg4), FWebJSParam(Arg5)};
+		Invoke(5, ArgArray);
+	}
+
+	template<typename T1, typename T2, typename T3, typename T4, typename T5, typename T6>
+	void operator()(T1 Arg1, T2 Arg2, T3 Arg3, T4 Arg4, T5 Arg5, T6 Arg6) const
+	{
+		FWebJSParam ArgArray[6] = {FWebJSParam(Arg1), FWebJSParam(Arg2), FWebJSParam(Arg3), FWebJSParam(Arg4), FWebJSParam(Arg5), FWebJSParam(Arg6)};
+		Invoke(6, ArgArray);
+	}
+
+	template<typename T1, typename T2, typename T3, typename T4, typename T5, typename T6, typename T7>
+	void operator()(T1 Arg1, T2 Arg2, T3 Arg3, T4 Arg4, T5 Arg5, T6 Arg6, T7 Arg7) const
+	{
+		FWebJSParam ArgArray[7] = {FWebJSParam(Arg1), FWebJSParam(Arg2), FWebJSParam(Arg3), FWebJSParam(Arg4), FWebJSParam(Arg5), FWebJSParam(Arg6), FWebJSParam(Arg7)};
+		Invoke(7, ArgArray);
+	}
+#endif
+
+private:
+
+	void Invoke(int32 ArgCount, FWebJSParam Arguments[]) const;
+
+	TWeakPtr<FWebJSScripting> ScriptingPtr;
+	FGuid FunctionId;
+};
\ No newline at end of file
diff --git a/Engine/Source/Runtime/WebBrowser/WebBrowser.Build.cs b/Engine/Source/Runtime/WebBrowser/WebBrowser.Build.cs
--- a/Engine/Source/Runtime/WebBrowser/WebBrowser.Build.cs
+++ b/Engine/Source/Runtime/WebBrowser/WebBrowser.Build.cs
@@ -13,10 +13,12 @@ public class WebBrowser : ModuleRules
 			new string[]
 			{
 				"Core",
+				"CoreUObject",
 				"RHI",
 				"InputCore",
-				"SlateCore",
 				"Slate",
+				"SlateCore",
+				"Serialization",
 				"CEF3Utils",
 			}
 		);
