Merging CL#2601386 using UE4-To-UE4-LauncherDev:
Fix 32bit compilation errors

diff --git a/Engine/Source/Runtime/WebBrowser/Private/WebJSScripting.h b/Engine/Source/Runtime/WebBrowser/Private/WebJSScripting.h
--- a/Engine/Source/Runtime/WebBrowser/Private/WebJSScripting.h
+++ b/Engine/Source/Runtime/WebBrowser/Private/WebJSScripting.h
@@ -122,7 +122,7 @@ private:
 			UPTRINT IntPtr = reinterpret_cast<UPTRINT>(Ptr);
 			if (sizeof(UPTRINT) > 4)
 			{
-				Guid[0] ^= IntPtr >> 32;
+				Guid[0] ^= (static_cast<uint64>(IntPtr) >> 32);
 			}
 			Guid[1] ^= IntPtr & 0xFFFFFFFF;
 		}
@@ -135,7 +135,7 @@ private:
 		UPTRINT IntPtr = 0;
 		if (sizeof(UPTRINT) > 4)
 		{
-			IntPtr = (UPTRINT)(Guid[0] ^ BaseGuid[0]) << 32;
+			IntPtr = static_cast<UPTRINT>(static_cast<uint64>(Guid[0] ^ BaseGuid[0]) << 32);
 		}
 		IntPtr |= (Guid[1] ^ BaseGuid[1]) & 0xFFFFFFFF;
 
diff --git a/Engine/Source/Runtime/WebBrowser/Private/WebJSStructDeserializerBackend.h b/Engine/Source/Runtime/WebBrowser/Private/WebJSStructDeserializerBackend.h
--- a/Engine/Source/Runtime/WebBrowser/Private/WebJSStructDeserializerBackend.h
+++ b/Engine/Source/Runtime/WebBrowser/Private/WebJSStructDeserializerBackend.h
@@ -52,7 +52,7 @@ public:
 	virtual bool ReadProperty(TSharedPtr<FWebJSScripting> Scripting, UProperty* Property, UProperty* Outer, void* Data, int32 ArrayIndex ) override;
 
 	CefRefPtr<CefListValue> List;
-	int32 Index;
+	size_t Index;
 };
 
 class FCefDictionaryValueWalker
@@ -72,7 +72,7 @@ public:
 
 private:
 	CefRefPtr<CefDictionaryValue> Dictionary;
-	int32 Index;
+	size_t Index;
 	CefDictionaryValue::KeyList Keys;
 };
 
