[INTEGRATE] Merging CL#2604183 using UE4-To-UE4-LauncherDev:
Allow using the native mouse cursor handle from CEF.
JIRA: OPP-3734

diff --git a/Engine/Source/Runtime/Core/Private/Linux/LinuxCursor.cpp b/Engine/Source/Runtime/Core/Private/Linux/LinuxCursor.cpp
--- a/Engine/Source/Runtime/Core/Private/Linux/LinuxCursor.cpp
+++ b/Engine/Source/Runtime/Core/Private/Linux/LinuxCursor.cpp
@@ -34,6 +34,7 @@ FLinuxCursor::FLinuxCursor()
 		switch( CurCursorIndex )
 		{
 		case EMouseCursor::None:
+		case EMouseCursor::Custom:
 			// The mouse cursor will not be visible when None is used
 			break;
 
@@ -149,6 +150,7 @@ FLinuxCursor::~FLinuxCursor()
 			case EMouseCursor::GrabHand:
 			case EMouseCursor::GrabHandClosed:
 			case EMouseCursor::SlashedCircle:
+			case EMouseCursor::Custom:
 				// Standard shared cursors don't need to be destroyed
 				break;
 			case EMouseCursor::EyeDropper:
@@ -163,6 +165,11 @@ FLinuxCursor::~FLinuxCursor()
 	}
 }
 
+void FLinuxCursor::SetCustomShape( SDL_HCursor CursorHandle )
+{
+	CursorHandles[EMouseCursor::Custom] = CursorHandle;
+}
+
 FVector2D FLinuxCursor::GetPosition() const
 {
 	if (!bPositionCacheIsValid)
@@ -212,7 +219,7 @@ void FLinuxCursor::SetType( const EMouseCursor::Type InNewCursor )
 {	
 	checkf( InNewCursor < EMouseCursor::TotalCursorCount, TEXT("Invalid cursor(%d) supplied"), InNewCursor );
 	CurrentType = InNewCursor;
-	if(InNewCursor == EMouseCursor::None)
+	if(CursorHandles[InNewCursor] == NULL)
 	{
 		bHidden = true;
 		SDL_ShowCursor(SDL_DISABLE);
diff --git a/Engine/Source/Runtime/Core/Private/Mac/MacCursor.cpp b/Engine/Source/Runtime/Core/Private/Mac/MacCursor.cpp
--- a/Engine/Source/Runtime/Core/Private/Mac/MacCursor.cpp
+++ b/Engine/Source/Runtime/Core/Private/Mac/MacCursor.cpp
@@ -41,6 +41,7 @@ FMacCursor::FMacCursor()
 		switch( CurCursorIndex )
 		{
 			case EMouseCursor::None:
+			case EMouseCursor::Custom:
 				break;
 
 			case EMouseCursor::Default:
@@ -193,7 +194,11 @@ FMacCursor::~FMacCursor()
 			case EMouseCursor::ResizeSouthWest:
 			case EMouseCursor::CardinalCross:
 			case EMouseCursor::EyeDropper:
-				[CursorHandles[CurCursorIndex] release];
+			case EMouseCursor::Custom:
+				if (CursorHandles[CurCursorIndex] != NULL)
+				{
+					[CursorHandles[CurCursorIndex] release];
+				}
 				break;
 
 			default:
@@ -204,6 +209,17 @@ FMacCursor::~FMacCursor()
 	}
 }
 
+void FMacCursor::SetCustomShape( NSCursor* CursorHandle )
+{
+	SCOPED_AUTORELEASE_POOL;
+	[CursorHandle retain];
+	if (CursorHandles[EMouseCursor::Custom] != NULL)
+	{
+		[CursorHandles[EMouseCursor::Custom] release];
+	}
+	CursorHandles[EMouseCursor::Custom] = CursorHandle;
+}
+
 FVector2D FMacCursor::GetPosition() const
 {
 	FVector2D CurrentPos = CurrentPosition;
diff --git a/Engine/Source/Runtime/Core/Private/Windows/WindowsCursor.cpp b/Engine/Source/Runtime/Core/Private/Windows/WindowsCursor.cpp
--- a/Engine/Source/Runtime/Core/Private/Windows/WindowsCursor.cpp
+++ b/Engine/Source/Runtime/Core/Private/Windows/WindowsCursor.cpp
@@ -15,6 +15,7 @@ FWindowsCursor::FWindowsCursor()
 		switch( CurCursorIndex )
 		{
 		case EMouseCursor::None:
+		case EMouseCursor::Custom:
 			// The mouse cursor will not be visible when None is used
 			break;
 
@@ -119,6 +120,7 @@ FWindowsCursor::~FWindowsCursor()
 		case EMouseCursor::GrabHandClosed:
 		case EMouseCursor::SlashedCircle:
 		case EMouseCursor::EyeDropper:
+		case EMouseCursor::Custom:
 			// Standard shared cursors don't need to be destroyed
 			break;
 
@@ -130,6 +132,11 @@ FWindowsCursor::~FWindowsCursor()
 	}
 }
 
+void FWindowsCursor::SetCustomShape( HCURSOR CursorHandle )
+{
+	CursorHandles[EMouseCursor::Custom] = CursorHandle;
+}
+
 FVector2D FWindowsCursor::GetPosition() const
 {
 	POINT CursorPos;
diff --git a/Engine/Source/Runtime/Core/Public/GenericPlatform/ICursor.h b/Engine/Source/Runtime/Core/Public/GenericPlatform/ICursor.h
--- a/Engine/Source/Runtime/Core/Public/GenericPlatform/ICursor.h
+++ b/Engine/Source/Runtime/Core/Public/GenericPlatform/ICursor.h
@@ -64,6 +64,9 @@ namespace EMouseCursor
 		/** Eye-dropper cursor for picking colors */
 		EyeDropper,
 
+		/** Custom cursor shape for platforms that support setting a native cursor shape. Same as specifying None if not set. */
+		Custom,
+
 		/** Number of cursors we support */
 		TotalCursorCount
 	};
diff --git a/Engine/Source/Runtime/Core/Public/Linux/LinuxCursor.h b/Engine/Source/Runtime/Core/Public/Linux/LinuxCursor.h
--- a/Engine/Source/Runtime/Core/Public/Linux/LinuxCursor.h
+++ b/Engine/Source/Runtime/Core/Public/Linux/LinuxCursor.h
@@ -35,6 +35,13 @@ public:
 
 public:
 
+	/**
+	 * Defines a custom cursor shape for the EMouseCursor::Custom type.
+	 * 
+	 * @param CursorHandle	A native cursor handle to show when EMouseCursor::Custom is selected.
+	 */
+	virtual void SetCustomShape( SDL_HCursor CursorHandle );
+
 	bool UpdateCursorClipping( FVector2D& CursorPosition );
 
 	bool IsHidden();
diff --git a/Engine/Source/Runtime/Core/Public/Mac/MacCursor.h b/Engine/Source/Runtime/Core/Public/Mac/MacCursor.h
--- a/Engine/Source/Runtime/Core/Public/Mac/MacCursor.h
+++ b/Engine/Source/Runtime/Core/Public/Mac/MacCursor.h
@@ -36,6 +36,13 @@ public:
 
 public:
 
+	/**
+	 * Defines a custom cursor shape for the EMouseCursor::Custom type.
+	 *
+	 * @param CursorHandle	A native NSCursor instance to show when EMouseCursor::Custom is selected.
+	 */
+	virtual void SetCustomShape( NSCursor* CursorHandle );
+
 	bool UpdateCursorClipping( FVector2D& CursorPosition );
 
 	void WarpCursor( const int32 X, const int32 Y );
diff --git a/Engine/Source/Runtime/Core/Public/Windows/WindowsCursor.h b/Engine/Source/Runtime/Core/Public/Windows/WindowsCursor.h
--- a/Engine/Source/Runtime/Core/Public/Windows/WindowsCursor.h
+++ b/Engine/Source/Runtime/Core/Public/Windows/WindowsCursor.h
@@ -34,6 +34,14 @@ public:
 
 	virtual void Lock( const RECT* const Bounds ) override;
 
+public:
+
+	/**
+	 * Defines a custom cursor shape for the EMouseCursor::Custom type.
+	 * 
+	 * @param CursorHandle	A native cursor handle to show when EMouseCursor::Custom is selected.
+	 */
+	virtual void SetCustomShape( HCURSOR CursorHandle );
 
 private:
 
diff --git a/Engine/Source/Runtime/SlateCore/Public/Application/SlateApplicationBase.h b/Engine/Source/Runtime/SlateCore/Public/Application/SlateApplicationBase.h
--- a/Engine/Source/Runtime/SlateCore/Public/Application/SlateApplicationBase.h
+++ b/Engine/Source/Runtime/SlateCore/Public/Application/SlateApplicationBase.h
@@ -465,6 +465,11 @@ public:
 	 */
 	bool IsSlateAsleep();
 
+	TSharedPtr<ICursor> GetPlatformCursor()
+	{
+		return PlatformApplication->Cursor;
+	}
+
 protected:
 
 	// Gets set when Slate goes to sleep and cleared when active.
diff --git a/Engine/Source/Runtime/WebBrowser/Private/WebBrowserWindow.cpp b/Engine/Source/Runtime/WebBrowser/Private/WebBrowserWindow.cpp
--- a/Engine/Source/Runtime/WebBrowser/Private/WebBrowserWindow.cpp
+++ b/Engine/Source/Runtime/WebBrowser/Private/WebBrowserWindow.cpp
@@ -14,6 +14,15 @@
 #include <AppKit/NSEvent.h>
 #endif
 
+#if PLATFORM_WINDOWS
+#include "WindowsCursor.h"
+typedef FWindowsCursor FPlatformCursor;
+#elif PLATFORM_MAC
+#include "MacCursor.h"
+typedef FMacCursor FPlatformCursor;
+#else
+#endif
+
 FWebBrowserWindow::FWebBrowserWindow(FIntPoint InViewportSize, FString InUrl, TOptional<FString> InContentsToLoad, bool InShowErrorMessage, bool InThumbMouseButtonNavigation, bool InUseTransparency)
 	: DocumentState(EWebBrowserDocumentState::NoDocument)
 	, UpdatableTexture(nullptr)
@@ -780,10 +789,27 @@ void FWebBrowserWindow::OnPaint(CefRenderHandler::PaintElementType Type, const C
 void FWebBrowserWindow::OnCursorChange(CefCursorHandle CefCursor, CefRenderHandler::CursorType Type, const CefCursorInfo& CustomCursorInfo)
 {
 	switch (Type) {
+		// Map the basic 3 cursor types directly to Slate types on all platforms
 		case CT_NONE:
 			Cursor = EMouseCursor::None;
 			break;
+		case CT_POINTER:
+			Cursor = EMouseCursor::Default;
+			break;
 		case CT_IBEAM:
+			Cursor = EMouseCursor::TextEditBeam;
+			break;
+		#if PLATFORM_WINDOWS || PLATFORM_MAC
+		// Platform specific support for native cursor types
+		default:
+			{
+				FPlatformCursor* PlatformCursor = (FPlatformCursor*)FSlateApplication::Get().GetPlatformCursor().Get();
+				PlatformCursor->SetCustomShape(CefCursor);
+				Cursor = EMouseCursor::Custom;
+			}
+			break;
+		#else
+		// Map to closest Slate equivalent on platforms where native cursors are not available.
 		case CT_VERTICALTEXT:
 			Cursor = EMouseCursor::TextEditBeam;
 			break;
@@ -837,12 +863,11 @@ void FWebBrowserWindow::OnCursorChange(CefCursorHandle CefCursor, CefRenderHandl
 		case CT_NODROP:
 			Cursor = EMouseCursor::SlashedCircle;
 			break;
-		case CT_POINTER:
 		default:
 			Cursor = EMouseCursor::Default;
 			break;
+		#endif
 	}
-
 	// Tell Slate to update the cursor now
 	FSlateApplication::Get().QueryCursor();
 }
