[INTEGRATE] Merging CL#2605806 using UE4-To-UE4-LauncherDev:
Initialize function parameters before deserializing from JS data.
OPP-3750

diff --git a/Engine/Source/Runtime/WebBrowser/Private/WebJSScripting.cpp b/Engine/Source/Runtime/WebBrowser/Private/WebJSScripting.cpp
--- a/Engine/Source/Runtime/WebBrowser/Private/WebJSScripting.cpp
+++ b/Engine/Source/Runtime/WebBrowser/Private/WebJSScripting.cpp
@@ -366,10 +366,13 @@ bool FWebJSScripting::HandleExecuteUObjectMethodMessage(CefRefPtr<CefListValue>
 
 		if (ParamsSize > 0)
 		{
-			Params.AddZeroed(ParamsSize);
 			// UFunction is a subclass of UStruct, so we can treat the arguments as a struct for deserialization
+			UStruct* TypeInfo = Cast<UStruct>(Function);
+
+			Params.AddUninitialized(ParamsSize);
+			TypeInfo->InitializeStruct(Params.GetData());
 			FWebJSStructDeserializerBackend Backend = FWebJSStructDeserializerBackend(SharedThis(this), NamedArgs);
-			FStructDeserializer::Deserialize(Params.GetData(), *Cast<UStruct>(Function), Backend);
+			FStructDeserializer::Deserialize(Params.GetData(), *TypeInfo, Backend);
 		}
 	}
 
diff --git a/Engine/Source/Runtime/WebBrowser/Private/WebJSStructDeserializerBackend.cpp b/Engine/Source/Runtime/WebBrowser/Private/WebJSStructDeserializerBackend.cpp
--- a/Engine/Source/Runtime/WebBrowser/Private/WebJSStructDeserializerBackend.cpp
+++ b/Engine/Source/Runtime/WebBrowser/Private/WebJSStructDeserializerBackend.cpp
@@ -69,7 +69,6 @@ namespace {
 		}
 	}
 
-
 	/**
 	 * Sets the value of the given property.
 	 *
@@ -102,7 +101,7 @@ namespace {
 		{
 			UPropertyType* TypedProperty = Cast<UPropertyType>(Property);
 
-			if (TypedProperty == nullptr)
+			if (TypedProperty == nullptr || ArrayIndex >= TypedProperty->ArrayDim)
 			{
 				return false;
 			}
@@ -224,8 +223,6 @@ namespace {
 			|| ReadNumericProperty<UDoubleProperty>(Property, Outer, Data, ArrayIndex, Container, Key)
 			|| ReadJSFunctionProperty(Scripting, Property, Outer, Data, ArrayIndex, Container, Key);
 	}
-
-
 }
 
 TSharedPtr<ICefContainerWalker> FCefListValueWalker::GetNextToken(EStructDeserializerBackendTokens& OutToken, FString& PropertyName)
