[INTEGRATE] Merging CL#2606002 using UE4-To-UE4-LauncherDev, which removes support for the CEF MessageRouter mechanism, as the UObject script bindings are more convenient and flexible.

The feature has never been included in an official release, so removing it should not cause a bw compatibility problem.

diff --git a/Engine/Source/Programs/UnrealCEFSubProcess/Private/UnrealCEFSubProcessApp.cpp b/Engine/Source/Programs/UnrealCEFSubProcess/Private/UnrealCEFSubProcessApp.cpp
--- a/Engine/Source/Programs/UnrealCEFSubProcess/Private/UnrealCEFSubProcessApp.cpp
+++ b/Engine/Source/Programs/UnrealCEFSubProcess/Private/UnrealCEFSubProcessApp.cpp
@@ -7,32 +7,19 @@
 
 FUnrealCEFSubProcessApp::FUnrealCEFSubProcessApp()
 {
-    CefMessageRouterConfig MessageRouterConfig;
-    MessageRouterConfig.js_query_function = "ueQuery";
-    MessageRouterConfig.js_cancel_function = "ueQueryCancel";
-    MessageRouter = CefMessageRouterRendererSide::Create(MessageRouterConfig);
 }
 
-
-void FUnrealCEFSubProcessApp::OnContextCreated(CefRefPtr<CefBrowser> Browser,
-    CefRefPtr<CefFrame> Frame,
-    CefRefPtr<CefV8Context> Context)
+void FUnrealCEFSubProcessApp::OnContextCreated( CefRefPtr<CefBrowser> Browser, CefRefPtr<CefFrame> Frame, CefRefPtr<CefV8Context> Context )
 {
 	RemoteScripting.OnContextCreated(Browser, Frame, Context);
-	MessageRouter->OnContextCreated(Browser, Frame, Context);
 }
 
-void FUnrealCEFSubProcessApp::OnContextReleased(CefRefPtr<CefBrowser> Browser,
-    CefRefPtr<CefFrame> Frame,
-    CefRefPtr<CefV8Context> Context)
+void FUnrealCEFSubProcessApp::OnContextReleased( CefRefPtr<CefBrowser> Browser, CefRefPtr<CefFrame> Frame, CefRefPtr<CefV8Context> Context )
 {
 	RemoteScripting.OnContextReleased(Browser, Frame, Context);
-    MessageRouter->OnContextReleased(Browser, Frame, Context);
 }
 
-bool FUnrealCEFSubProcessApp::OnProcessMessageReceived(CefRefPtr<CefBrowser> Browser,
-    CefProcessId SourceProcess,
-    CefRefPtr<CefProcessMessage> Message)
+bool FUnrealCEFSubProcessApp::OnProcessMessageReceived( CefRefPtr<CefBrowser> Browser, CefProcessId SourceProcess, CefRefPtr<CefProcessMessage> Message )
 {
 	bool Result = false;
 	FString MessageName = Message->GetName().ToWString().c_str();
@@ -40,14 +27,11 @@ bool FUnrealCEFSubProcessApp::OnProcessMessageReceived(CefRefPtr<CefBrowser> Bro
 	{
 		Result = RemoteScripting.OnProcessMessageReceived(Browser, SourceProcess, Message);
 	}
-	else
-	{
-		Result = MessageRouter->OnProcessMessageReceived(Browser, SourceProcess, Message);
-	}
+
 	return Result;
 }
 
-void FUnrealCEFSubProcessApp::OnRenderThreadCreated(CefRefPtr<CefListValue> ExtraInfo)
+void FUnrealCEFSubProcessApp::OnRenderThreadCreated( CefRefPtr<CefListValue> ExtraInfo )
 {
 	for(size_t I=0; I<ExtraInfo->GetSize(); I++)
 	{
@@ -64,5 +48,4 @@ void FUnrealCEFSubProcessApp::OnRenderThreadCreated(CefRefPtr<CefListValue> Extr
 	}
 }
 
-
 #endif
\ No newline at end of file
diff --git a/Engine/Source/Programs/UnrealCEFSubProcess/Private/UnrealCEFSubProcessApp.h b/Engine/Source/Programs/UnrealCEFSubProcess/Private/UnrealCEFSubProcessApp.h
--- a/Engine/Source/Programs/UnrealCEFSubProcess/Private/UnrealCEFSubProcessApp.h
+++ b/Engine/Source/Programs/UnrealCEFSubProcess/Private/UnrealCEFSubProcessApp.h
@@ -10,7 +10,6 @@
 #pragma push_macro("OVERRIDE")
 #undef OVERRIDE // cef headers provide their own OVERRIDE macro
 #include "include/cef_app.h"
-#include "include/wrapper/cef_message_router.h"
 #pragma pop_macro("OVERRIDE")
 
 #if PLATFORM_WINDOWS
@@ -23,42 +22,34 @@
  * Implements CEF App and other Process level interfaces
  */
 class FUnrealCEFSubProcessApp
-    : public CefApp
-    , public CefRenderProcessHandler
+	: public CefApp
+	, public CefRenderProcessHandler
 {
 public:
     
-    /**
-     * Default Constructor
-     */
-    FUnrealCEFSubProcessApp();
+	/**
+	 * Default Constructor
+	 */
+	FUnrealCEFSubProcessApp();
+
 private:
-    // CefApp methods:
-    virtual CefRefPtr<CefRenderProcessHandler> GetRenderProcessHandler() override { return this; }
+	// CefApp methods:
+	virtual CefRefPtr<CefRenderProcessHandler> GetRenderProcessHandler() override { return this; }
 
-    // CefRenderProcessHandler methods:
-    virtual void OnContextCreated(CefRefPtr<CefBrowser> Browser,
-        CefRefPtr<CefFrame> Frame,
-        CefRefPtr<CefV8Context> Context) override;
+	// CefRenderProcessHandler methods:
+	virtual void OnContextCreated( CefRefPtr<CefBrowser> Browser, CefRefPtr<CefFrame> Frame, CefRefPtr<CefV8Context> Context ) override;
 
-    virtual void OnContextReleased(CefRefPtr<CefBrowser> Browser,
-        CefRefPtr<CefFrame> Frame,
-        CefRefPtr<CefV8Context> Context) override;
+	virtual void OnContextReleased( CefRefPtr<CefBrowser> Browser, CefRefPtr<CefFrame> Frame, CefRefPtr<CefV8Context> Context ) override;
 
-    virtual bool OnProcessMessageReceived(CefRefPtr<CefBrowser> Browser,
-        CefProcessId SourceProcess,
-        CefRefPtr<CefProcessMessage> Message) override;
+	virtual bool OnProcessMessageReceived( CefRefPtr<CefBrowser> Browser, CefProcessId SourceProcess, CefRefPtr<CefProcessMessage> Message ) override;
 
-	virtual void OnRenderThreadCreated(CefRefPtr<CefListValue> ExtraInfo) override;
+	virtual void OnRenderThreadCreated( CefRefPtr<CefListValue> ExtraInfo ) override;
 
 	// Handles remote scripting messages from the frontend process
 	FUnrealCEFSubProcessRemoteScripting RemoteScripting;
 
-    // Routes messages from Javascript to the browser process and back
-    CefRefPtr<CefMessageRouterRendererSide> MessageRouter;
-
-    // Include the default reference counting implementation.
-    IMPLEMENT_REFCOUNTING(FUnrealCEFSubProcessApp);
+	// Include the default reference counting implementation.
+	IMPLEMENT_REFCOUNTING(FUnrealCEFSubProcessApp);
 };
 
 #endif // WITH_CEF3
\ No newline at end of file
diff --git a/Engine/Source/Runtime/WebBrowser/Private/SWebBrowser.cpp b/Engine/Source/Runtime/WebBrowser/Private/SWebBrowser.cpp
--- a/Engine/Source/Runtime/WebBrowser/Private/SWebBrowser.cpp
+++ b/Engine/Source/Runtime/WebBrowser/Private/SWebBrowser.cpp
@@ -22,8 +22,6 @@ void SWebBrowser::Construct(const FArguments& InArgs, const TSharedPtr<IWebBrows
 	OnTitleChanged = InArgs._OnTitleChanged;
 	OnUrlChanged = InArgs._OnUrlChanged;
 	OnBeforeNavigation = InArgs._OnBeforeNavigation;
-	OnJSQueryReceived = InArgs._OnJSQueryReceived;
-	OnJSQueryCanceled = InArgs._OnJSQueryCanceled;
 	OnLoadUrl = InArgs._OnLoadUrl;
 	OnBeforePopup = InArgs._OnBeforePopup;
 	OnCreateWindow = InArgs._OnCreateWindow;
@@ -155,8 +153,6 @@ void SWebBrowser::Construct(const FArguments& InArgs, const TSharedPtr<IWebBrows
 		BrowserWindow->OnToolTip().AddSP(this, &SWebBrowser::HandleToolTip);
 		BrowserWindow->OnBeforeBrowse().BindSP(this, &SWebBrowser::HandleBeforeNavigation);
 		BrowserWindow->OnLoadUrl().BindSP(this, &SWebBrowser::HandleLoadUrl);
-		BrowserWindow->OnJSQueryReceived().BindSP(this, &SWebBrowser::HandleJSQueryReceived);
-		BrowserWindow->OnJSQueryCanceled().BindSP(this, &SWebBrowser::HandleJSQueryCanceled);
 		BrowserWindow->OnBeforePopup().BindSP(this, &SWebBrowser::HandleBeforePopup);
 		BrowserViewport = MakeShareable(new FWebBrowserViewport(BrowserWindow, ViewportWidget));
 		ViewportWidget->SetViewportInterface(BrowserViewport.ToSharedRef());
@@ -408,20 +404,6 @@ bool SWebBrowser::HandleLoadUrl(const FString& Method, const FString& Url, FStri
 	return false;
 }
 
-bool SWebBrowser::HandleJSQueryReceived( int64 QueryId, FString QueryString, bool Persistent, FJSQueryResultDelegate Delegate )
-{
-	if (OnJSQueryReceived.IsBound())
-	{
-		return OnJSQueryReceived.Execute(QueryId, QueryString, Persistent, Delegate);
-	}
-	return false;
-}
-
-void SWebBrowser::HandleJSQueryCanceled( int64 QueryId )
-{
-	OnJSQueryCanceled.ExecuteIfBound(QueryId);
-}
-
 bool SWebBrowser::HandleBeforePopup(FString URL, FString Target)
 {
 	if (OnBeforePopup.IsBound())
diff --git a/Engine/Source/Runtime/WebBrowser/Private/WebBrowserHandler.cpp b/Engine/Source/Runtime/WebBrowser/Private/WebBrowserHandler.cpp
--- a/Engine/Source/Runtime/WebBrowser/Private/WebBrowserHandler.cpp
+++ b/Engine/Source/Runtime/WebBrowser/Private/WebBrowserHandler.cpp
@@ -11,12 +11,6 @@
 FWebBrowserHandler::FWebBrowserHandler()
 	: ShowErrorMessage(true)
 {
-	// This has to match the config in UnrealCEFSubpProcess
-	CefMessageRouterConfig MessageRouterConfig;
-	MessageRouterConfig.js_query_function = "ueQuery";
-	MessageRouterConfig.js_cancel_function = "ueQueryCancel";
-	MessageRouter = CefMessageRouterBrowserSide::Create(MessageRouterConfig);
-	MessageRouter->AddHandler(this, true);
 }
 
 void FWebBrowserHandler::OnTitleChange(CefRefPtr<CefBrowser> Browser, const CefString& Title)
@@ -99,7 +93,6 @@ void FWebBrowserHandler::OnAfterCreated(CefRefPtr<CefBrowser> Browser)
 
 void FWebBrowserHandler::OnBeforeClose(CefRefPtr<CefBrowser> Browser)
 {
-	MessageRouter->OnBeforeClose(Browser);
 	TSharedPtr<FWebBrowserWindow> BrowserWindow = BrowserWindowPtr.Pin();
 
 	if (BrowserWindow.IsValid())
@@ -266,7 +259,6 @@ bool FWebBrowserHandler::OnBeforeResourceLoad(CefRefPtr<CefBrowser> Browser, Cef
 
 void FWebBrowserHandler::OnRenderProcessTerminated(CefRefPtr<CefBrowser> Browser, TerminationStatus Status)
 {
-	MessageRouter->OnRenderProcessTerminated(Browser);
 }
 
 bool FWebBrowserHandler::OnBeforeBrowse(CefRefPtr<CefBrowser> Browser,
@@ -283,7 +275,6 @@ bool FWebBrowserHandler::OnBeforeBrowse(CefRefPtr<CefBrowser> Browser,
 		}
 	}
 
-	MessageRouter->OnBeforeBrowse(Browser, Frame);
 	return false;
 }
 
@@ -317,37 +308,9 @@ bool FWebBrowserHandler::OnProcessMessageReceived(CefRefPtr<CefBrowser> Browser,
 	{
 		Retval = BrowserWindow->OnProcessMessageReceived(Browser, SourceProcess, Message);
 	}
-	if (! Retval)
-	{
-		Retval = MessageRouter->OnProcessMessageReceived(Browser, SourceProcess, Message);
-	}
 	return Retval;
 }
 
-bool FWebBrowserHandler::OnQuery(CefRefPtr<CefBrowser> Browser,
-		CefRefPtr<CefFrame> Frame,
-		int64 QueryId,
-		const CefString& Request,
-		bool Persistent,
-		CefRefPtr<CefMessageRouterBrowserSide::Callback> Callback)
-{
-	TSharedPtr<FWebBrowserWindow> BrowserWindow = BrowserWindowPtr.Pin();
-	if (BrowserWindow.IsValid())
-	{
-		return BrowserWindow->OnQuery(QueryId, Request, Persistent, Callback);
-	}
-	return false;
-}
-
-void FWebBrowserHandler::OnQueryCanceled(CefRefPtr<CefBrowser> Browser, CefRefPtr<CefFrame> Frame, int64 QueryId)
-{
-	TSharedPtr<FWebBrowserWindow> BrowserWindow = BrowserWindowPtr.Pin();
-	if (BrowserWindow.IsValid())
-	{
-		BrowserWindow->OnQueryCanceled(QueryId);
-	}
-}
-
 bool FWebBrowserHandler::OnKeyEvent(CefRefPtr<CefBrowser> Browser,
 	const CefKeyEvent& Event,
 	CefEventHandle OsEvent)
diff --git a/Engine/Source/Runtime/WebBrowser/Private/WebBrowserHandler.h b/Engine/Source/Runtime/WebBrowser/Private/WebBrowserHandler.h
--- a/Engine/Source/Runtime/WebBrowser/Private/WebBrowserHandler.h
+++ b/Engine/Source/Runtime/WebBrowser/Private/WebBrowserHandler.h
@@ -11,7 +11,6 @@
 #pragma push_macro("OVERRIDE")
 #undef OVERRIDE // cef headers provide their own OVERRIDE macro
 #include "include/cef_client.h"
-#include "include/wrapper/cef_message_router.h"
 #pragma pop_macro("OVERRIDE")
 
 #if PLATFORM_WINDOWS
@@ -33,7 +32,6 @@ class FWebBrowserHandler
 	, public CefRenderHandler
 	, public CefRequestHandler
 	, public CefKeyboardHandler
-	, public CefMessageRouterBrowserSide::Handler
 {
 public:
 
@@ -173,19 +171,6 @@ public:
 		CefRefPtr<CefRequest> Request,
 		bool IsRedirect) override;
 	virtual CefRefPtr<CefResourceHandler> GetResourceHandler( CefRefPtr<CefBrowser> Browser, CefRefPtr< CefFrame > Frame, CefRefPtr<CefRequest> Request ) override;
-public:
-	// CefMessageRouterBrowserSide::Handler Interface
-	
-	virtual bool OnQuery(CefRefPtr<CefBrowser> Browser,
-		CefRefPtr<CefFrame> Frame,
-		int64 QueryId,
-		const CefString& Request,
-		bool Persistent,
-		CefRefPtr<CefMessageRouterBrowserSide::Callback> Callback) override;
-
-	virtual void OnQueryCanceled(CefRefPtr<CefBrowser> Browser,
-		CefRefPtr<CefFrame> Frame,
-		int64 QueryId) override;
 
 public:
 	// CefKeyboardHandler interface
@@ -204,9 +189,6 @@ private:
 	/** Whether to show an error message in case of loading errors. */
 	bool ShowErrorMessage;
 
-	/** The message router is used as a part of a generic message api between Javascript in the render process and the application process */
-	CefRefPtr<CefMessageRouterBrowserSide> MessageRouter;
-
 	// Include the default reference counting implementation.
 	IMPLEMENT_REFCOUNTING(FWebBrowserHandler);
 };
diff --git a/Engine/Source/Runtime/WebBrowser/Private/WebBrowserWindow.cpp b/Engine/Source/Runtime/WebBrowser/Private/WebBrowserWindow.cpp
--- a/Engine/Source/Runtime/WebBrowser/Private/WebBrowserWindow.cpp
+++ b/Engine/Source/Runtime/WebBrowser/Private/WebBrowserWindow.cpp
@@ -1073,35 +1073,6 @@ void FWebBrowserWindow::UnbindUObject(const FString& Name, UObject* Object, bool
 	Scripting->UnbindUObject(Name, Object, bIsPermanent);
 }
 
-bool FWebBrowserWindow::OnQuery(int64 QueryId, const CefString& Request, bool Persistent, CefRefPtr<CefMessageRouterBrowserSide::Callback> Callback)
-{
-	if (OnJSQueryReceived().IsBound())
-	{
-		FString QueryString = Request.ToWString().c_str();
-		FJSQueryResultDelegate Delegate = FJSQueryResultDelegate::CreateLambda(
-			[Callback] (int ErrorCode, FString Message)
-			{
-				CefString MessageString = *Message;
-				if (ErrorCode == 0)
-				{
-					Callback->Success(MessageString);
-				}
-				else
-				{
-					Callback->Failure(ErrorCode, MessageString);
-				}
-			}
-		);
-		return OnJSQueryReceived().Execute(QueryId, QueryString, Persistent, Delegate);
-	}
-	return false;
-}
-
-void FWebBrowserWindow::OnQueryCanceled(int64 QueryId)
-{
-	OnJSQueryCanceled().ExecuteIfBound(QueryId);
-}
-
 bool FWebBrowserWindow::OnCefBeforePopup(const CefString& Target_Url, const CefString& Target_Frame_Name)
 {
 	if (OnBeforePopup().IsBound())
diff --git a/Engine/Source/Runtime/WebBrowser/Private/WebBrowserWindow.h b/Engine/Source/Runtime/WebBrowser/Private/WebBrowserWindow.h
--- a/Engine/Source/Runtime/WebBrowser/Private/WebBrowserWindow.h
+++ b/Engine/Source/Runtime/WebBrowser/Private/WebBrowserWindow.h
@@ -149,16 +149,6 @@ public:
 		return NeedsRedrawEvent;
 	}
 	
-	virtual FONJSQueryReceived& OnJSQueryReceived() override
-	{
-		return JSQueryReceivedDelegate;
-	}
-
-	virtual FONJSQueryCanceled& OnJSQueryCanceled() override
-	{
-		return JSQueryCanceledDelegate;
-	}
-
 	virtual FOnBeforeBrowse& OnBeforeBrowse() override
 	{
 		return BeforeBrowseDelegate;
@@ -260,25 +250,6 @@ private:
 	 */
 	bool OnProcessMessageReceived(CefRefPtr<CefBrowser> Browser, CefProcessId SourceProcess, CefRefPtr<CefProcessMessage> Message);
 
-	/**
-	 * Called when JavaScript code sends a message to the UE process.
-	 * Needs to return true or false to tell CEF wether the query is being handled by user code or not.
-	 *
-	 * @param QueryId A unique id for the query. Used to refer to it in OnQueryCanceled.
-	 * @param Request The query string itself as passed in from the JS code.
-	 * @param Persistent Is this a persistent query or not. If not, client code expects the callback to be invoked only once, wheras persistent queries are terminated by invoking Failure, Success can be invoked multiple times until then.
-	 * @param Callback A handle to pass data back to the JS code. 
-	 */
-	bool OnQuery(int64 QueryId, const CefString& Request, bool Persistent, CefRefPtr<CefMessageRouterBrowserSide::Callback> Callback);
-
-	/**
-	 * Called when an outstanding query has been canceled eother explicitly from JS code or implicitly by navigating away from the page containing the code.
-	 * Will only be called if OnQuery has previously returned true for the same QueryId.
-	 *
-	 * @param QueryId A unique id for the query. A handler should use it to locate and remove any handlers that might be in flight.
-	 */
-	void OnQueryCanceled(int64 QueryId);
-	
 	/**
 	 * Called before browser navigation.
 	 *
@@ -428,9 +399,6 @@ private:
 	/** Delegate for overriding Url contents. */
 	FOnLoadUrl LoadUrlDelegate;
 
-	FONJSQueryReceived JSQueryReceivedDelegate;
-	FONJSQueryCanceled JSQueryCanceledDelegate;
-
 	/** Delegate for notifying that a popup window is attempting to open. */
 	FOnBeforePopupDelegate BeforePopupDelegate;
 	
diff --git a/Engine/Source/Runtime/WebBrowser/Public/IWebBrowserWindow.h b/Engine/Source/Runtime/WebBrowser/Public/IWebBrowserWindow.h
--- a/Engine/Source/Runtime/WebBrowser/Public/IWebBrowserWindow.h
+++ b/Engine/Source/Runtime/WebBrowser/Public/IWebBrowserWindow.h
@@ -20,8 +20,6 @@ enum class EWebBrowserDocumentState
 	NoDocument
 };
 
-DECLARE_DELEGATE_TwoParams(FJSQueryResultDelegate, int, FString);
-
 /**
  * Interface for dealing with a Web Browser window
  */
@@ -249,15 +247,6 @@ public:
 	DECLARE_EVENT(IWebBrowserWindow, FOnNeedsRedraw)
 	virtual FOnNeedsRedraw& OnNeedsRedraw() = 0;
 	
-	/** A delegate that is invoked when JS code sends a query to the front end. The result delegate cand either be executed immediately or saved and executed later (multiple times if the boolean peristent argument is true). 
-	 * The arguments are an integer error code (0 for success) and a reply string. If you need pass more complex data to the JS code, you will have to pack the data in some way (such as JSON encoding it). */
-	DECLARE_DELEGATE_RetVal_FourParams(bool, FONJSQueryReceived, int64, FString, bool, FJSQueryResultDelegate)
-	virtual FONJSQueryReceived& OnJSQueryReceived() = 0;
-
-	/** A delegate that is invoked when an outstanding query is canceled. Implement this if you are saving delegates passed to OnQueryReceived. */
-	DECLARE_DELEGATE_OneParam(FONJSQueryCanceled, int64)
-	virtual FONJSQueryCanceled& OnJSQueryCanceled() = 0;
-	
 	/** A delegate that is invoked prior to browser navigation. */
 	DECLARE_DELEGATE_RetVal_TwoParams(bool, FOnBeforeBrowse, const FString& /*Url*/, bool /*bIsRedirect*/)
 	virtual FOnBeforeBrowse& OnBeforeBrowse() = 0;
diff --git a/Engine/Source/Runtime/WebBrowser/Public/SWebBrowser.h b/Engine/Source/Runtime/WebBrowser/Public/SWebBrowser.h
--- a/Engine/Source/Runtime/WebBrowser/Public/SWebBrowser.h
+++ b/Engine/Source/Runtime/WebBrowser/Public/SWebBrowser.h
@@ -9,13 +9,6 @@ class IWebBrowserWindow;
 class FWebBrowserViewport;
 class UObject;
 
-DECLARE_DELEGATE_TwoParams(FJSQueryResultDelegate, int, FString);
-DECLARE_DELEGATE_RetVal_FourParams(bool, FOnJSQueryReceivedDelegate, int64, FString, bool, FJSQueryResultDelegate);
-DECLARE_DELEGATE_OneParam(FOnJSQueryCanceledDelegate, int64);
-
-DECLARE_DELEGATE_TwoParams(FJSQueryResultDelegate, int, FString);
-DECLARE_DELEGATE_RetVal_FourParams(bool, FOnJSQueryReceivedDelegate, int64, FString, bool, FJSQueryResultDelegate);
-DECLARE_DELEGATE_OneParam(FOnJSQueryCanceledDelegate, int64);
 DECLARE_DELEGATE_RetVal_TwoParams(bool, FOnBeforePopupDelegate, FString, FString);
 DECLARE_DELEGATE_RetVal_OneParam(bool, FOnCreateWindowDelegate, const TWeakPtr<IWebBrowserWindow>&);
 
@@ -27,12 +20,13 @@ public:
 	DECLARE_DELEGATE_RetVal_ThreeParams(bool, FOnLoadUrl, const FString& /*Method*/, const FString& /*Url*/, FString& /* Response */)
 
 	SLATE_BEGIN_ARGS(SWebBrowser)
-		: _InitialURL(TEXT("www.google.com"))
+		: _InitialURL(TEXT("https://www.google.com"))
 		, _ShowControls(true)
 		, _ShowAddressBar(false)
 		, _ShowErrorMessage(true)
 		, _SupportsTransparency(false)
 		, _SupportsThumbMouseButtonNavigation(false)
+		, _BackgroundColor(255,255,255,255)
 		, _ViewportSize(FVector2D::ZeroVector)
 	{ }
 
@@ -81,12 +75,6 @@ public:
 		/** Called when the Url changes. */
 		SLATE_EVENT(FOnTextChanged, OnUrlChanged)
 	
-		/** Called when a custom Javascript message is received from the browser process. */
-		SLATE_EVENT(FOnJSQueryReceivedDelegate, OnJSQueryReceived)
-	
-		/** Called when a pending Javascript message has been canceled, either explicitly or by navigating away from the page containing the script. */
-		SLATE_EVENT(FOnJSQueryCanceledDelegate, OnJSQueryCanceled)
-		
 		/** Called before a popup window happens */
 		SLATE_EVENT(FOnBeforePopupDelegate, OnBeforePopup)
 
@@ -227,12 +215,6 @@ private:
 	/** Callback for showing browser tool tips. */
 	void HandleToolTip(FString ToolTipText);
 
-	/** Callback for received JS queries. */
-	bool HandleJSQueryReceived(int64 QueryId, FString QueryString, bool Persistent, FJSQueryResultDelegate ResultDelegate);
-
-	/** Callback for cancelled JS queries. */
-	void HandleJSQueryCanceled(int64 QueryId);
-
 	/**
 	 * A delegate that is executed prior to browser navigation.
 	 *
@@ -284,12 +266,6 @@ private:
 	/** A delegate that is invoked when document address changed. */
 	FOnTextChanged OnUrlChanged;
 	
-	/** A delegate that is invoked when render process Javascript code sends a query message to the client. */
-	FOnJSQueryReceivedDelegate OnJSQueryReceived;
-	
-	/** A delegate that is invoked when render process cancels an ongoing query. Handler must clean up corresponding result delegate. */
-	FOnJSQueryCanceledDelegate OnJSQueryCanceled;
-	
 	/** A delegate that is invoked when the browser attempts to pop up a new window */
 	FOnBeforePopupDelegate OnBeforePopup;
 
