Fixed a number of issues with web browser window spawning and closing.

Merging CL#2608123 LauncherDev->Main
change summary:
- Popup features/settings are now passed to when requesting creation of new browser window.
- Web browser handler now implements DoClose()
- Added support for overriding web browser close functionality.
- Launcher specific test for spawning browser windows with dynamic content.

diff --git a/Engine/Source/Runtime/WebBrowser/Private/SWebBrowser.cpp b/Engine/Source/Runtime/WebBrowser/Private/SWebBrowser.cpp
--- a/Engine/Source/Runtime/WebBrowser/Private/SWebBrowser.cpp
+++ b/Engine/Source/Runtime/WebBrowser/Private/SWebBrowser.cpp
@@ -14,6 +14,27 @@ SWebBrowser::SWebBrowser()
 {
 }
 
+SWebBrowser::~SWebBrowser()
+{
+	if (BrowserWindow.IsValid())
+	{
+		BrowserWindow->OnCreateWindow().Unbind();		
+		BrowserWindow->OnCloseWindow().Unbind();
+		BrowserWindow->OnDocumentStateChanged().RemoveAll(this);
+		BrowserWindow->OnNeedsRedraw().RemoveAll(this);
+		BrowserWindow->OnTitleChanged().RemoveAll(this);
+		BrowserWindow->OnUrlChanged().RemoveAll(this);
+		BrowserWindow->OnToolTip().RemoveAll(this);
+		BrowserWindow->OnBeforeBrowse().Unbind();
+		BrowserWindow->OnLoadUrl().Unbind();
+		BrowserWindow->OnBeforePopup().Unbind();
+		if(!BrowserWindow->IsClosing())
+		{
+			BrowserWindow->CloseBrowser(true);
+		}
+	}
+}
+
 void SWebBrowser::Construct(const FArguments& InArgs, const TSharedPtr<IWebBrowserWindow>& InWebBrowserWindow)
 {
 	OnLoadCompleted = InArgs._OnLoadCompleted;
@@ -25,6 +46,7 @@ void SWebBrowser::Construct(const FArguments& InArgs, const TSharedPtr<IWebBrows
 	OnLoadUrl = InArgs._OnLoadUrl;
 	OnBeforePopup = InArgs._OnBeforePopup;
 	OnCreateWindow = InArgs._OnCreateWindow;
+	OnCloseWindow = InArgs._OnCloseWindow;
 	AddressBarUrl = FText::FromString(InArgs._InitialURL);
 
 	BrowserWindow = InWebBrowserWindow;
@@ -146,6 +168,11 @@ void SWebBrowser::Construct(const FArguments& InArgs, const TSharedPtr<IWebBrows
 			BrowserWindow->OnCreateWindow().BindSP(this, &SWebBrowser::HandleCreateWindow);		
 		}
 
+		if(OnCloseWindow.IsBound())
+		{
+			BrowserWindow->OnCloseWindow().BindSP(this, &SWebBrowser::HandleCloseWindow);
+		}
+
 		BrowserWindow->OnDocumentStateChanged().AddSP(this, &SWebBrowser::HandleBrowserWindowDocumentStateChanged);
 		BrowserWindow->OnNeedsRedraw().AddSP(this, &SWebBrowser::HandleBrowserWindowNeedsRedraw);
 		BrowserWindow->OnTitleChanged().AddSP(this, &SWebBrowser::HandleTitleChanged);
@@ -422,11 +449,20 @@ void SWebBrowser::ExecuteJavascript(const FString& ScriptText)
 	}
 }
 
-bool SWebBrowser::HandleCreateWindow(const TWeakPtr<IWebBrowserWindow>& NewBrowserWindow)
+bool SWebBrowser::HandleCreateWindow(const TWeakPtr<IWebBrowserWindow>& NewBrowserWindow, const TWeakPtr<IWebBrowserPopupFeatures>& PopupFeatures)
 {
 	if(OnCreateWindow.IsBound())
 	{
-		return OnCreateWindow.Execute(NewBrowserWindow);
+		return OnCreateWindow.Execute(NewBrowserWindow, PopupFeatures);
+	}
+	return false;
+}
+
+bool SWebBrowser::HandleCloseWindow(const TWeakPtr<IWebBrowserWindow>& NewBrowserWindow)
+{
+	if(OnCloseWindow.IsBound())
+	{
+		return OnCloseWindow.Execute(NewBrowserWindow);
 	}
 	return false;
 }
diff --git a/Engine/Source/Runtime/WebBrowser/Private/WebBrowserHandler.cpp b/Engine/Source/Runtime/WebBrowser/Private/WebBrowserHandler.cpp
--- a/Engine/Source/Runtime/WebBrowser/Private/WebBrowserHandler.cpp
+++ b/Engine/Source/Runtime/WebBrowser/Private/WebBrowserHandler.cpp
@@ -4,6 +4,8 @@
 #include "WebBrowserHandler.h"
 #include "WebBrowserWindow.h"
 #include "WebBrowserSingleton.h"
+#include "WebBrowserPopupFeatures.h"
+
 
 #define LOCTEXT_NAMESPACE "WebBrowserHandler"
 
@@ -77,10 +79,10 @@ void FWebBrowserHandler::OnAfterCreated(CefRefPtr<CefBrowser> Browser)
 				}
 
 				// Request a UI window for the browser.  If it is not created we do some cleanup.
-				bool bUIWindowCreated = BrowserWindowParent->RequestCreateWindow(NewBrowserWindow.ToSharedRef());
+				bool bUIWindowCreated = BrowserWindowParent->RequestCreateWindow(NewBrowserWindow.ToSharedRef(), BrowserPopupFeatures);
 				if(!bUIWindowCreated)
 				{
-					NewBrowserWindow->CloseBrowser();
+					NewBrowserWindow->CloseBrowser(true);
 				}
 			}
 			else
@@ -91,13 +93,23 @@ void FWebBrowserHandler::OnAfterCreated(CefRefPtr<CefBrowser> Browser)
 	}
 }
 
+ bool FWebBrowserHandler::DoClose(CefRefPtr<CefBrowser> Browser)
+ {
+	 TSharedPtr<FWebBrowserWindow> BrowserWindow = BrowserWindowPtr.Pin();
+	if(BrowserWindow.IsValid())
+	{
+		BrowserWindow->OnBrowserClosing();
+	}
+
+	return false;
+ }
+
 void FWebBrowserHandler::OnBeforeClose(CefRefPtr<CefBrowser> Browser)
 {
 	TSharedPtr<FWebBrowserWindow> BrowserWindow = BrowserWindowPtr.Pin();
-
 	if (BrowserWindow.IsValid())
 	{
-		BrowserWindow->BindCefBrowser(nullptr);
+		BrowserWindow->OnBrowserClosed();
 	}
 }
 
@@ -126,6 +138,7 @@ bool FWebBrowserHandler::OnBeforePopup( CefRefPtr<CefBrowser> Browser,
 			{
 				CefRefPtr<FWebBrowserHandler> NewHandler(new FWebBrowserHandler);
 				NewHandler->SetBrowserWindowParent(BrowserWindow);
+				NewHandler->SetPopupFeatures(MakeShareable(new FWebBrowserPopupFeatures(PopupFeatures)));
 				Client = NewHandler;
 				
 				CefWindowHandle ParentWindowHandle = BrowserWindow->GetCefBrowser()->GetHost()->GetWindowHandle();
diff --git a/Engine/Source/Runtime/WebBrowser/Private/WebBrowserHandler.h b/Engine/Source/Runtime/WebBrowser/Private/WebBrowserHandler.h
--- a/Engine/Source/Runtime/WebBrowser/Private/WebBrowserHandler.h
+++ b/Engine/Source/Runtime/WebBrowser/Private/WebBrowserHandler.h
@@ -19,6 +19,7 @@
 
 
 class FWebBrowserWindow;
+class FWebBrowserPopupFeatures;
 
 
 /**
@@ -54,6 +55,16 @@ public:
 	 */
 	void SetBrowserWindowParent(TSharedPtr<FWebBrowserWindow> InBrowserWindow);
 
+	/**
+	 * Sets the browser window features and settings for popups which will be passed along when creating the new window.
+	 *
+	 * @param InPopupFeatures The popup features and settings for the window.
+	 */
+	void SetPopupFeatures(const TSharedPtr<FWebBrowserPopupFeatures>& InPopupFeatures)
+	{
+		BrowserPopupFeatures = InPopupFeatures;
+	}
+
 	/** Sets whether to show messages on loading errors. */
 	void SetShowErrorMessage(bool InShowErrorMessage)
 	{
@@ -113,6 +124,7 @@ public:
 	// CefLifeSpanHandler Interface
 
 	virtual void OnAfterCreated(CefRefPtr<CefBrowser> Browser) override;
+	virtual bool DoClose(CefRefPtr<CefBrowser> Browser) override;
 	virtual void OnBeforeClose(CefRefPtr<CefBrowser> Browser) override;
 	virtual bool OnBeforePopup(CefRefPtr<CefBrowser> Browser, 
 		CefRefPtr<CefFrame> Frame, 
@@ -186,6 +198,9 @@ private:
 	/** Weak Pointer to the parent web browser window */
 	TWeakPtr<FWebBrowserWindow> BrowserWindowParentPtr;
 
+	/** Stores popup window features and settings */
+	TSharedPtr<FWebBrowserPopupFeatures> BrowserPopupFeatures;
+
 	/** Whether to show an error message in case of loading errors. */
 	bool ShowErrorMessage;
 
diff --git a/Engine/Source/Runtime/WebBrowser/Private/WebBrowserPopupFeatures.cpp b/Engine/Source/Runtime/WebBrowser/Private/WebBrowserPopupFeatures.cpp
new file mode 100644
--- /dev/null
+++ b/Engine/Source/Runtime/WebBrowser/Private/WebBrowserPopupFeatures.cpp
@@ -0,0 +1,150 @@
+// Copyright 1998-2015 Epic Games, Inc. All Rights Reserved.
+
+#include "WebBrowserPrivatePCH.h"
+#include "WebBrowserPopupFeatures.h"
+
+#if WITH_CEF3
+
+FWebBrowserPopupFeatures::FWebBrowserPopupFeatures()
+	: X(0)
+	, bXSet(false)
+	, Y(0)
+	, bYSet(false)
+	, Width(0)
+	, bWidthSet(false)
+	, Height(0)
+	, bHeightSet(false)
+	, bMenuBarVisible(true)
+	, bStatusBarVisible(false)
+	, bToolBarVisible(true)
+	, bLocationBarVisible(true)
+	, bScrollbarsVisible(true)
+	, bResizable(true)
+	, bIsFullscreen(false)
+	, bIsDialog(false)
+{
+}
+
+
+FWebBrowserPopupFeatures::FWebBrowserPopupFeatures( const CefPopupFeatures& PopupFeatures )
+{
+	X = PopupFeatures.x;
+	bXSet = PopupFeatures.xSet ? true : false;
+	Y = PopupFeatures.y;
+	bYSet = PopupFeatures.ySet ? true : false;
+	Width = PopupFeatures.width;
+	bWidthSet = PopupFeatures.widthSet ? true : false;
+	Height = PopupFeatures.height;
+	bHeightSet = PopupFeatures.heightSet ? true : false;
+	bMenuBarVisible = PopupFeatures.menuBarVisible ? true : false;
+	bStatusBarVisible = PopupFeatures.statusBarVisible ? true : false;
+	bToolBarVisible = PopupFeatures.toolBarVisible ? true : false;
+	bLocationBarVisible = PopupFeatures.locationBarVisible ? true : false;
+	bScrollbarsVisible = PopupFeatures.scrollbarsVisible ? true : false;
+	bResizable = PopupFeatures.resizable ? true : false;
+	bIsFullscreen = PopupFeatures.fullscreen ? true : false;
+	bIsDialog = PopupFeatures.dialog ? true : false;
+
+	int Count = PopupFeatures.additionalFeatures ? cef_string_list_size(PopupFeatures.additionalFeatures) : 0;
+	CefString ListValue;
+
+	for(int ListIdx = 0; ListIdx < Count; ListIdx++) 
+	{
+		cef_string_list_value(PopupFeatures.additionalFeatures, ListIdx, ListValue.GetWritableStruct());
+		AdditionalFeatures.Add(ListValue.ToWString().c_str());
+	}
+
+}
+
+FWebBrowserPopupFeatures::~FWebBrowserPopupFeatures()
+{
+}
+
+int FWebBrowserPopupFeatures::GetX() const 
+{
+	return X;
+}
+
+bool FWebBrowserPopupFeatures::IsXSet() const 
+{
+	return bXSet;
+}
+
+int FWebBrowserPopupFeatures::GetY() const 
+{
+	return Y;
+}
+
+bool FWebBrowserPopupFeatures::IsYSet() const 
+{
+	return bYSet;
+}
+
+int FWebBrowserPopupFeatures::GetWidth() const 
+{
+	return Width;
+}
+
+bool FWebBrowserPopupFeatures::IsWidthSet() const 
+{
+	return bWidthSet;
+}
+
+int FWebBrowserPopupFeatures::GetHeight() const 
+{
+	return Height;
+}
+
+bool FWebBrowserPopupFeatures::IsHeightSet() const 
+{
+	return bHeightSet;
+}
+
+bool FWebBrowserPopupFeatures::IsMenuBarVisible() const 
+{
+	return bMenuBarVisible;
+}
+
+bool FWebBrowserPopupFeatures::IsStatusBarVisible() const 
+{
+	return bStatusBarVisible;
+}
+
+bool FWebBrowserPopupFeatures::IsToolBarVisible() const 
+{
+	return bToolBarVisible;
+}
+
+bool FWebBrowserPopupFeatures::IsLocationBarVisible() const 
+{
+	return bLocationBarVisible;
+}
+
+bool FWebBrowserPopupFeatures::IsScrollbarsVisible() const 
+{
+	return bScrollbarsVisible;
+}
+
+bool FWebBrowserPopupFeatures::IsResizable() const 
+{
+	return bResizable;
+}
+
+bool FWebBrowserPopupFeatures::IsFullscreen() const 
+{
+	return bIsFullscreen;
+}
+
+bool FWebBrowserPopupFeatures::IsDialog() const 
+{
+	return bIsDialog;
+}
+
+TArray<FString> FWebBrowserPopupFeatures::GetAdditionalFeatures() const 
+{
+	return AdditionalFeatures;
+}
+
+#endif
+
+
diff --git a/Engine/Source/Runtime/WebBrowser/Private/WebBrowserPopupFeatures.h b/Engine/Source/Runtime/WebBrowser/Private/WebBrowserPopupFeatures.h
new file mode 100644
--- /dev/null
+++ b/Engine/Source/Runtime/WebBrowser/Private/WebBrowserPopupFeatures.h
@@ -0,0 +1,73 @@
+// Copyright 1998-2015 Epic Games, Inc. All Rights Reserved.
+
+#pragma once
+
+#if WITH_CEF3
+
+#include "IWebBrowserPopupFeatures.h"
+
+#if PLATFORM_WINDOWS
+#include "AllowWindowsPlatformTypes.h"
+#endif
+
+#pragma push_macro("OVERRIDE")
+#undef OVERRIDE // cef headers provide their own OVERRIDE macro
+#include "include/cef_base.h"
+#pragma pop_macro("OVERRIDE")
+
+#if PLATFORM_WINDOWS
+#include "HideWindowsPlatformTypes.h"
+#endif
+
+class FWebBrowserPopupFeatures
+	: public IWebBrowserPopupFeatures
+{
+public:
+	FWebBrowserPopupFeatures();
+	FWebBrowserPopupFeatures(const CefPopupFeatures& PopupFeatures);
+	virtual ~FWebBrowserPopupFeatures();
+
+	// IWebBrowserPopupFeatures Interface
+
+	virtual int GetX() const override;
+	virtual bool IsXSet() const override;
+	virtual int GetY() const override;
+	virtual bool IsYSet() const override;
+	virtual int GetWidth() const override;
+	virtual bool IsWidthSet() const override;
+	virtual int GetHeight() const override;
+	virtual bool IsHeightSet() const override;
+	virtual bool IsMenuBarVisible() const override;
+	virtual bool IsStatusBarVisible() const override;
+	virtual bool IsToolBarVisible() const override;
+	virtual bool IsLocationBarVisible() const override;
+	virtual bool IsScrollbarsVisible() const override;
+	virtual bool IsResizable() const override;
+	virtual bool IsFullscreen() const override;
+	virtual bool IsDialog() const override;
+	virtual TArray<FString> GetAdditionalFeatures() const override;
+
+private:
+
+	int X;
+	bool bXSet;
+	int Y;
+	bool bYSet;
+	int Width;
+	bool bWidthSet;
+	int Height;
+	bool bHeightSet;
+
+	bool bMenuBarVisible;
+	bool bStatusBarVisible;
+	bool bToolBarVisible;
+	bool bLocationBarVisible;
+	bool bScrollbarsVisible;
+	bool bResizable;
+
+	bool bIsFullscreen;
+	bool bIsDialog;
+	TArray<FString> AdditionalFeatures;
+};
+
+#endif
\ No newline at end of file
diff --git a/Engine/Source/Runtime/WebBrowser/Private/WebBrowserSingleton.cpp b/Engine/Source/Runtime/WebBrowser/Private/WebBrowserSingleton.cpp
--- a/Engine/Source/Runtime/WebBrowser/Private/WebBrowserSingleton.cpp
+++ b/Engine/Source/Runtime/WebBrowser/Private/WebBrowserSingleton.cpp
@@ -129,7 +129,7 @@ FWebBrowserSingleton::~FWebBrowserSingleton()
 	{
 		if (WindowInterfaces[Index].IsValid())
 		{
-			WindowInterfaces[Index].Pin()->CloseBrowser();
+			WindowInterfaces[Index].Pin()->CloseBrowser(true);
 		}
 	}
 	// Just in case, although we deallocate WebBrowserApp right after this.
@@ -230,8 +230,11 @@ bool FWebBrowserSingleton::Tick(float DeltaTime)
 		else if (bIsSlateAwake) // only check for Tick activity if Slate is currently ticking
 		{
 			TSharedPtr<FWebBrowserWindow> BrowserWindow = WindowInterfaces[Index].Pin();
-			// Test if we've ticked recently. If not assume the browser window has become hidden.
-			BrowserWindow->CheckTickActivity();
+			if(BrowserWindow.IsValid())
+			{
+				// Test if we've ticked recently. If not assume the browser window has become hidden.
+				BrowserWindow->CheckTickActivity();
+			}
 		}
 	}
 	CefDoMessageLoopWork();
diff --git a/Engine/Source/Runtime/WebBrowser/Private/WebBrowserWindow.cpp b/Engine/Source/Runtime/WebBrowser/Private/WebBrowserWindow.cpp
--- a/Engine/Source/Runtime/WebBrowser/Private/WebBrowserWindow.cpp
+++ b/Engine/Source/Runtime/WebBrowser/Private/WebBrowserWindow.cpp
@@ -3,6 +3,7 @@
 #include "WebBrowserPrivatePCH.h"
 #include "WebBrowserWindow.h"
 #include "WebBrowserByteResource.h"
+#include "WebBrowserPopupFeatures.h"
 #include "WebJSScripting.h"
 #include "RHI.h"
 
@@ -49,7 +50,7 @@ FWebBrowserWindow::FWebBrowserWindow(FIntPoint InViewportSize, FString InUrl, TO
 
 FWebBrowserWindow::~FWebBrowserWindow()
 {
-	CloseBrowser();
+	CloseBrowser(true);
 
 	if (FSlateApplication::IsInitialized() && FSlateApplication::Get().GetRenderer().IsValid() && UpdatableTexture != nullptr)
 	{
@@ -434,11 +435,11 @@ bool FWebBrowserWindow::SupportsNewWindows()
 	return OnCreateWindow().IsBound() ? true : false;
 }
 
-bool FWebBrowserWindow::RequestCreateWindow( const TSharedRef<IWebBrowserWindow>& NewBrowserWindow )
+bool FWebBrowserWindow::RequestCreateWindow( const TSharedRef<IWebBrowserWindow>& NewBrowserWindow, const TSharedPtr<IWebBrowserPopupFeatures>& BrowserPopupFeatures )
 {
 	if(OnCreateWindow().IsBound())
 	{
-		return OnCreateWindow().Execute(TWeakPtr<IWebBrowserWindow>(NewBrowserWindow));
+		return OnCreateWindow().Execute(TWeakPtr<IWebBrowserWindow>(NewBrowserWindow), TWeakPtr<IWebBrowserPopupFeatures>(BrowserPopupFeatures));
 	}
 	return false;
 }
@@ -678,15 +679,11 @@ void FWebBrowserWindow::SetHandler(CefRefPtr<FWebBrowserHandler> InHandler)
 	}
 }
 
-void FWebBrowserWindow::CloseBrowser()
+void FWebBrowserWindow::CloseBrowser(bool bForce)
 {
-	bIsClosing = true;
 	if (IsValid())
 	{
-		Scripting->UnbindCefBrowser();
-		InternalCefBrowser->GetHost()->CloseBrowser(false);
-		InternalCefBrowser = nullptr;
-		Handler = nullptr;
+		InternalCefBrowser->GetHost()->CloseBrowser(bForce);
 	}
 }
 
@@ -1085,4 +1082,21 @@ bool FWebBrowserWindow::OnCefBeforePopup(const CefString& Target_Url, const CefS
 	return false;
 }
 
+void FWebBrowserWindow::OnBrowserClosing()
+{
+	bIsClosing = true;
+}
+
+void FWebBrowserWindow::OnBrowserClosed()
+{
+	if(OnCloseWindow().IsBound())
+	{
+		OnCloseWindow().Execute(TWeakPtr<IWebBrowserWindow>(SharedThis(this)));
+	}
+
+	Scripting->UnbindCefBrowser();
+	InternalCefBrowser = nullptr;
+	Handler = nullptr;
+}
+
 #endif
diff --git a/Engine/Source/Runtime/WebBrowser/Private/WebBrowserWindow.h b/Engine/Source/Runtime/WebBrowser/Private/WebBrowserWindow.h
--- a/Engine/Source/Runtime/WebBrowser/Private/WebBrowserWindow.h
+++ b/Engine/Source/Runtime/WebBrowser/Private/WebBrowserWindow.h
@@ -115,7 +115,7 @@ public:
 	virtual void Reload() override;
 	virtual void StopLoad() override;
 	virtual void ExecuteJavascript(const FString& Script) override;
-	virtual void CloseBrowser() override;
+	virtual void CloseBrowser(bool bForce) override;
 	virtual void BindUObject(const FString& Name, UObject* Object, bool bIsPermanent = true) override;
 	virtual void UnbindUObject(const FString& Name, UObject* Object = nullptr, bool bIsPermanent = true) override;
 
@@ -164,6 +164,11 @@ public:
 		return CreateWindowDelegate;
 	}
 
+	virtual FOnCloseWindow& OnCloseWindow() override
+	{
+		return CloseWindowDelegate;
+	}
+
 	virtual FCursorReply OnCursorQuery( const FGeometry& MyGeometry, const FPointerEvent& CursorEvent ) override
 	{
 		return Cursor == EMouseCursor::Default ? FCursorReply::Unhandled() : FCursorReply::Cursor(Cursor);
@@ -277,13 +282,28 @@ private:
 	 * @return true if window creation functionality was provided, false otherwise.  If false, RequestCreateWindow() will always return false.
 	 */
 	bool SupportsNewWindows();
-	
+
 	/** Called when the browser requests a new UI window
 	 *
 	 * @param NewBrowserWindow The web browser window to display in the new UI window.
+	 * @param BrowserPopupFeatures The popup features and settings for the browser window.
 	 * @return true if the UI window was created, false otherwise.
 	 */
-	bool RequestCreateWindow(const TSharedRef<IWebBrowserWindow>& NewBrowserWindow);
+	bool RequestCreateWindow(const TSharedRef<IWebBrowserWindow>& NewBrowserWindow, const TSharedPtr<IWebBrowserPopupFeatures>& BrowserPopupFeatures);
+	
+	//bool SupportsCloseWindows();
+	//bool RequestCloseWindow(const TSharedRef<IWebBrowserWindow>& BrowserWindow);
+
+
+	/**
+	 * Called once the browser begins closing procedures.
+	 */
+	void OnBrowserClosing();
+	
+	/**
+	 * Called once the browser is closed.
+	 */
+	void OnBrowserClosed();
 
 public:
 
@@ -405,6 +425,9 @@ private:
 	/** Delegate for handaling requests to create new windows. */
 	FOnCreateWindow CreateWindowDelegate;
 
+	/** Delegate for handaling requests to close new windows that were created. */
+	FOnCloseWindow CloseWindowDelegate;
+
 	/** Tracks the current mouse cursor */
 	EMouseCursor::Type Cursor;
 
diff --git a/Engine/Source/Runtime/WebBrowser/Public/IWebBrowserPopupFeatures.h b/Engine/Source/Runtime/WebBrowser/Public/IWebBrowserPopupFeatures.h
new file mode 100644
--- /dev/null
+++ b/Engine/Source/Runtime/WebBrowser/Public/IWebBrowserPopupFeatures.h
@@ -0,0 +1,27 @@
+// Copyright 1998-2015 Epic Games, Inc. All Rights Reserved.
+
+#pragma once
+
+class IWebBrowserPopupFeatures
+{
+public:
+	virtual ~IWebBrowserPopupFeatures() {};
+
+	virtual int GetX() const = 0;
+	virtual bool IsXSet() const = 0;
+	virtual int GetY() const = 0;
+	virtual bool IsYSet() const = 0;
+	virtual int GetWidth() const = 0;
+	virtual bool IsWidthSet() const = 0;
+	virtual int GetHeight() const = 0;
+	virtual bool IsHeightSet() const = 0;
+	virtual bool IsMenuBarVisible() const = 0;
+	virtual bool IsStatusBarVisible() const = 0;
+	virtual bool IsToolBarVisible() const = 0;
+	virtual bool IsLocationBarVisible() const = 0;
+	virtual bool IsScrollbarsVisible() const = 0;
+	virtual bool IsResizable() const = 0;
+	virtual bool IsFullscreen() const = 0;
+	virtual bool IsDialog() const = 0;
+	virtual TArray<FString> GetAdditionalFeatures() const = 0;
+};
\ No newline at end of file
diff --git a/Engine/Source/Runtime/WebBrowser/Public/IWebBrowserWindow.h b/Engine/Source/Runtime/WebBrowser/Public/IWebBrowserWindow.h
--- a/Engine/Source/Runtime/WebBrowser/Public/IWebBrowserWindow.h
+++ b/Engine/Source/Runtime/WebBrowser/Public/IWebBrowserWindow.h
@@ -10,7 +10,7 @@ struct FPointerEvent;
 class FReply;
 class FCursorReply;
 class FSlateShaderResource;
-
+class IWebBrowserPopupFeatures;
 
 enum class EWebBrowserDocumentState
 {
@@ -201,9 +201,13 @@ public:
 
 	/** Execute Javascript on the page. */
 	virtual void ExecuteJavascript(const FString& Script) = 0;
-	
-	/** Close this window so that it can no longer be used. */
-	virtual void CloseBrowser() = 0;
+
+	/**
+	 * Close this window so that it can no longer be used.
+	 *
+	 * @param bForce Designates whether the web browser close should be forced.
+	 */
+	virtual void CloseBrowser(bool bForce) = 0;
 
 	/** 
 	 * Expose a UObject instance to the browser runtime.
@@ -259,10 +263,14 @@ public:
 	DECLARE_DELEGATE_RetVal_TwoParams(bool, FOnBeforePopupDelegate, FString, FString);
 	virtual FOnBeforePopupDelegate& OnBeforePopup() = 0;
 
-	/** A delegate this is invoked when an existing browser requests creation of a new browser window. */
-	DECLARE_DELEGATE_RetVal_OneParam(bool, FOnCreateWindow, const TWeakPtr<IWebBrowserWindow>& /*NewBrowserWindow*/)
+	/** A delegate that is invoked when an existing browser requests creation of a new browser window. */
+	DECLARE_DELEGATE_RetVal_TwoParams(bool, FOnCreateWindow, const TWeakPtr<IWebBrowserWindow>& /*NewBrowserWindow*/, const TWeakPtr<IWebBrowserPopupFeatures>& /* PopupFeatures*/)
 	virtual FOnCreateWindow& OnCreateWindow() = 0;
 
+	/** A delegate that is invoked when closing created popup windows. */
+	DECLARE_DELEGATE_RetVal_OneParam(bool, FOnCloseWindow, const TWeakPtr<IWebBrowserWindow>& /*BrowserWindow*/)
+	virtual FOnCloseWindow& OnCloseWindow() = 0;
+
 protected:
 
 	/** Virtual Destructor. */
diff --git a/Engine/Source/Runtime/WebBrowser/Public/SWebBrowser.h b/Engine/Source/Runtime/WebBrowser/Public/SWebBrowser.h
--- a/Engine/Source/Runtime/WebBrowser/Public/SWebBrowser.h
+++ b/Engine/Source/Runtime/WebBrowser/Public/SWebBrowser.h
@@ -8,9 +8,12 @@ enum class EWebBrowserDocumentState;
 class IWebBrowserWindow;
 class FWebBrowserViewport;
 class UObject;
+class IWebBrowserPopupFeatures;
 
 DECLARE_DELEGATE_RetVal_TwoParams(bool, FOnBeforePopupDelegate, FString, FString);
-DECLARE_DELEGATE_RetVal_OneParam(bool, FOnCreateWindowDelegate, const TWeakPtr<IWebBrowserWindow>&);
+DECLARE_DELEGATE_RetVal_TwoParams(bool, FOnCreateWindowDelegate, const TWeakPtr<IWebBrowserWindow>&, const TWeakPtr<IWebBrowserPopupFeatures>&);
+DECLARE_DELEGATE_RetVal_OneParam(bool, FOnCloseWindowDelegate, const TWeakPtr<IWebBrowserWindow>&);
+
 
 class WEBBROWSER_API SWebBrowser
 	: public SCompoundWidget
@@ -81,6 +84,9 @@ public:
 		/** Called when the browser requests the creation of a new window */
 		SLATE_EVENT(FOnCreateWindowDelegate, OnCreateWindow)
 
+		/** Called when a browser window close event is detected */
+		SLATE_EVENT(FOnCloseWindowDelegate, OnCloseWindow)
+
 		/** Called before browser navigation. */
 		SLATE_EVENT(FOnBeforeBrowse, OnBeforeNavigation)
 	
@@ -92,6 +98,8 @@ public:
 	/** Default constructor. */
 	SWebBrowser();
 
+	~SWebBrowser();
+
 	/**
 	 * Construct the widget.
 	 *
@@ -229,7 +237,14 @@ private:
 	 *
 	 * @return true if if the window request was handled, false if the browser requesting the new window should be closed.
 	 */
-	bool HandleCreateWindow(const TWeakPtr<IWebBrowserWindow>& NewBrowserWindow);
+	bool HandleCreateWindow(const TWeakPtr<IWebBrowserWindow>& NewBrowserWindow, const TWeakPtr<IWebBrowserPopupFeatures>& PopupFeatures);
+
+	/**
+	 * A delegate that is executed when closing the browser window.
+	 *
+	 * @return true if if the window close was handled, false otherwise.
+	 */
+	bool HandleCloseWindow(const TWeakPtr<IWebBrowserWindow>& BrowserWindow);
 
 	/** Callback for popup window permission */
 	bool HandleBeforePopup(FString URL, FString Target);
@@ -272,6 +287,9 @@ private:
 	/** A delegate that is invoked when the browser requests a UI window for another browser it spawned */
 	FOnCreateWindowDelegate OnCreateWindow;
 
+	/** A delegate that is invoked when a window close event is detected */
+	FOnCloseWindowDelegate OnCloseWindow;
+
 	/** A delegate that is invoked prior to browser navigation */
 	FOnBeforeBrowse OnBeforeNavigation;
 	
