[INTEGRATE] Merging CL 2608342 and 2608482 using UE4-To-UE4-LauncherDev. Basic support for opening Chrome Dev Tools from debug builds of the WebBrowser widget.

diff --git a/Engine/Source/Runtime/WebBrowser/Private/WebBrowserHandler.cpp b/Engine/Source/Runtime/WebBrowser/Private/WebBrowserHandler.cpp
--- a/Engine/Source/Runtime/WebBrowser/Private/WebBrowserHandler.cpp
+++ b/Engine/Source/Runtime/WebBrowser/Private/WebBrowserHandler.cpp
@@ -328,6 +328,52 @@ bool FWebBrowserHandler::OnKeyEvent(CefRefPtr<CefBrowser> Browser,
 	const CefKeyEvent& Event,
 	CefEventHandle OsEvent)
 {
+#if UE_BUILD_DEBUG
+	// Show dev tools on CMD/CTRL+ALT+I
+	if( (Event.type == KEYEVENT_RAWKEYDOWN || Event.type == KEYEVENT_KEYDOWN) &&
+#if PLATFORM_MAC
+		(Event.modifiers == (EVENTFLAG_COMMAND_DOWN | EVENTFLAG_SHIFT_DOWN)) &&
+#else
+		(Event.modifiers == (EVENTFLAG_CONTROL_DOWN | EVENTFLAG_SHIFT_DOWN)) &&
+#endif
+		(Event.unmodified_character == 'i' || Event.unmodified_character == 'I')
+	  )
+	{
+
+		CefPoint Point;
+		CefString TargetUrl = "chrome-devtools://devtools/devtools.html";
+		CefString TargetFrameName = "devtools";
+		CefPopupFeatures PopupFeatures;
+		CefWindowInfo WindowInfo;
+		CefRefPtr<CefClient> NewClient;
+		CefBrowserSettings BrowserSettings;
+		bool NoJavascriptAccess = false;
+
+		PopupFeatures.xSet = false;
+		PopupFeatures.ySet = false;
+		PopupFeatures.heightSet = false;
+		PopupFeatures.widthSet = false;
+		PopupFeatures.locationBarVisible = false;
+		PopupFeatures.menuBarVisible = false;
+		PopupFeatures.toolBarVisible  = false;
+		PopupFeatures.statusBarVisible  = false;
+		PopupFeatures.resizable = true;
+
+		// Set max framerate to maximum supported.
+		BrowserSettings.windowless_frame_rate = 60;
+		// Disable plugins
+		BrowserSettings.plugins = STATE_DISABLED;
+
+		// OnBeforePopup already takes care of all the details required to ask the host application to create a new browser window.
+		bool bSuppressWindowCreation = OnBeforePopup(Browser, Browser->GetFocusedFrame(), TargetUrl, TargetFrameName, PopupFeatures, WindowInfo, NewClient, BrowserSettings, &NoJavascriptAccess);
+
+		if(! bSuppressWindowCreation)
+		{
+			Browser->GetHost()->ShowDevTools(WindowInfo, NewClient, BrowserSettings, Point);
+		}
+	}
+#endif
+
 #if PLATFORM_MAC
 	// We need to handle standard Copy/Paste/etc... shortcuts on OS X
 	if( (Event.type == KEYEVENT_RAWKEYDOWN || Event.type == KEYEVENT_KEYDOWN) &&
@@ -368,7 +414,6 @@ bool FWebBrowserHandler::OnKeyEvent(CefRefPtr<CefBrowser> Browser,
 		}
 	}
 #endif
-
 	TSharedPtr<FWebBrowserWindow> BrowserWindow = BrowserWindowPtr.Pin();
 	if (BrowserWindow.IsValid())
 	{
