Added a -nocef flag option to the launcher.
Merging using UE4-To-UE4-LauncherDev

diff --git a/Engine/Source/Runtime/WebBrowser/Private/SWebBrowser.cpp b/Engine/Source/Runtime/WebBrowser/Private/SWebBrowser.cpp
--- a/Engine/Source/Runtime/WebBrowser/Private/SWebBrowser.cpp
+++ b/Engine/Source/Runtime/WebBrowser/Private/SWebBrowser.cpp
@@ -59,6 +59,9 @@ void SWebBrowser::Construct(const FArguments& InArgs, const TSharedPtr<IWebBrows
 			OSWindowHandle = NativeWindow->GetOSWindowHandle();
 		}
 
+		static bool AllowCEF = !FParse::Param(FCommandLine::Get(), TEXT("nocef"));
+		if (AllowCEF)
+		{
 		BrowserWindow = IWebBrowserModule::Get().GetSingleton()->CreateBrowserWindow(
 			OSWindowHandle,
 			InArgs._InitialURL,
@@ -71,6 +74,7 @@ void SWebBrowser::Construct(const FArguments& InArgs, const TSharedPtr<IWebBrows
 			InArgs._BackgroundColor
 			);		
 	}
+	}
 
 
 	ChildSlot
@@ -184,6 +188,10 @@ void SWebBrowser::Construct(const FArguments& InArgs, const TSharedPtr<IWebBrows
 		BrowserViewport = MakeShareable(new FWebBrowserViewport(BrowserWindow, ViewportWidget));
 		ViewportWidget->SetViewportInterface(BrowserViewport.ToSharedRef());
 	}
+	else
+	{
+		OnLoadError.ExecuteIfBound();
+	}
 }
 
 void SWebBrowser::LoadURL(FString NewURL)
@@ -204,8 +212,11 @@ void SWebBrowser::LoadString(FString Contents, FString DummyURL)
 
 void SWebBrowser::Reload()
 {
+	if (BrowserWindow.IsValid())
+	{
 	BrowserWindow->Reload();
 }
+}
 
 FText SWebBrowser::GetTitleText() const
 {
@@ -218,7 +229,7 @@ FText SWebBrowser::GetTitleText() const
 
 FString SWebBrowser::GetUrl() const
 {
-	if (BrowserWindow->IsValid())
+	if (BrowserWindow.IsValid())
 	{
 		return BrowserWindow->GetUrl();
 	}
