[INTEGRATE] Merging CL#2611173 using UE4-To-UE4-LauncherDev

diff --git a/Engine/Source/Runtime/WebBrowser/Private/SWebBrowser.cpp b/Engine/Source/Runtime/WebBrowser/Private/SWebBrowser.cpp
--- a/Engine/Source/Runtime/WebBrowser/Private/SWebBrowser.cpp
+++ b/Engine/Source/Runtime/WebBrowser/Private/SWebBrowser.cpp
@@ -25,6 +25,8 @@ SWebBrowser::~SWebBrowser()
 		BrowserWindow->OnTitleChanged().RemoveAll(this);
 		BrowserWindow->OnUrlChanged().RemoveAll(this);
 		BrowserWindow->OnToolTip().RemoveAll(this);
+		BrowserWindow->OnShowPopup().RemoveAll(this);
+		BrowserWindow->OnDismissPopup().RemoveAll(this);
 		BrowserWindow->OnBeforeBrowse().Unbind();
 		BrowserWindow->OnLoadUrl().Unbind();
 		BrowserWindow->OnBeforePopup().Unbind();
@@ -185,6 +187,8 @@ void SWebBrowser::Construct(const FArguments& InArgs, const TSharedPtr<IWebBrows
 		BrowserWindow->OnBeforeBrowse().BindSP(this, &SWebBrowser::HandleBeforeNavigation);
 		BrowserWindow->OnLoadUrl().BindSP(this, &SWebBrowser::HandleLoadUrl);
 		BrowserWindow->OnBeforePopup().BindSP(this, &SWebBrowser::HandleBeforePopup);
+		BrowserWindow->OnShowPopup().AddSP(this, &SWebBrowser::HandleShowPopup);
+		BrowserWindow->OnDismissPopup().AddSP(this, &SWebBrowser::HandleDismissPopup);
 		BrowserViewport = MakeShareable(new FWebBrowserViewport(BrowserWindow, ViewportWidget));
 		ViewportWidget->SetViewportInterface(BrowserViewport.ToSharedRef());
 	}
@@ -494,5 +498,41 @@ void SWebBrowser::UnbindUObject(const FString& Name, UObject* Object, bool bIsPe
 	}
 }
 
+void SWebBrowser::HandleShowPopup(const FIntRect& PopupSize)
+{
+	TSharedPtr<SViewport> MenuContent;
+	SAssignNew(MenuContent, SViewport)
+				.ViewportSize(PopupSize.Size())
+				.EnableGammaCorrection(false)
+				.EnableBlending(false)
+				.IgnoreTextureAlpha(true)
+				.Visibility(EVisibility::Visible);
+	MenuViewport = MakeShareable(new FWebBrowserViewport(BrowserWindow, MenuContent, true));
+	MenuContent->SetViewportInterface(MenuViewport.ToSharedRef());
+	FWidgetPath WidgetPath;
+	FSlateApplication::Get().GeneratePathToWidgetUnchecked(ViewportWidget.ToSharedRef(), WidgetPath);
+	if (WidgetPath.IsValid())
+	{
+		TSharedRef< SWidget > MenuContentRef = MenuContent.ToSharedRef();
+		const FGeometry& BrowserGeometry = WidgetPath.Widgets.Last().Geometry;
+		const FVector2D NewPosition = BrowserGeometry.LocalToAbsolute(PopupSize.Min);
+
+		// Open the pop-up
+		TSharedPtr<IMenu> NewMenu = FSlateApplication::Get().PushMenu(ViewportWidget.ToSharedRef(), WidgetPath, MenuContentRef, NewPosition, FPopupTransitionEffect( FPopupTransitionEffect::ComboButton ), false);
+		PopupMenuPtr = NewMenu;
+		check(NewMenu.IsValid() && NewMenu->GetOwnedWindow().IsValid());
+	}
+
+}
+
+void SWebBrowser::HandleDismissPopup()
+{
+	if (PopupMenuPtr.IsValid())
+	{
+		PopupMenuPtr.Pin()->Dismiss();
+	}
+}
+
+
 #undef LOCTEXT_NAMESPACE
 
diff --git a/Engine/Source/Runtime/WebBrowser/Private/WebBrowserHandler.cpp b/Engine/Source/Runtime/WebBrowser/Private/WebBrowserHandler.cpp
--- a/Engine/Source/Runtime/WebBrowser/Private/WebBrowserHandler.cpp
+++ b/Engine/Source/Runtime/WebBrowser/Private/WebBrowserHandler.cpp
@@ -12,6 +12,7 @@
 #if WITH_CEF3
 FWebBrowserHandler::FWebBrowserHandler()
 	: ShowErrorMessage(true)
+	, bShowPopupRequested(false)
 {
 }
 
@@ -250,6 +251,34 @@ void FWebBrowserHandler::OnCursorChange(CefRefPtr<CefBrowser> Browser, CefCursor
 	
 }
 
+void FWebBrowserHandler::OnPopupShow(CefRefPtr<CefBrowser> Browser, bool bShow)
+{
+	TSharedPtr<FWebBrowserWindow> BrowserWindow = BrowserWindowPtr.Pin();
+	bShowPopupRequested = bShow;
+
+	if (BrowserWindow.IsValid())
+	{
+		if (!bShow)
+		{
+			BrowserWindow->HidePopup();
+		}
+	}
+
+}
+
+void FWebBrowserHandler::OnPopupSize(CefRefPtr<CefBrowser> Browser, const CefRect& Rect)
+{
+	if (Rect.width > 0 && Rect.height > 0)
+	{
+		PopupRect = Rect;
+		TSharedPtr<FWebBrowserWindow> BrowserWindow = BrowserWindowPtr.Pin();
+		if (BrowserWindow.IsValid() && bShowPopupRequested)
+		{
+			BrowserWindow->ShowPopup(PopupRect);
+			bShowPopupRequested = false;
+		}
+	}
+}
 
 bool FWebBrowserHandler::OnBeforeResourceLoad(CefRefPtr<CefBrowser> Browser, CefRefPtr<CefFrame> Frame, CefRefPtr<CefRequest> Request)
 {
diff --git a/Engine/Source/Runtime/WebBrowser/Private/WebBrowserHandler.h b/Engine/Source/Runtime/WebBrowser/Private/WebBrowserHandler.h
--- a/Engine/Source/Runtime/WebBrowser/Private/WebBrowserHandler.h
+++ b/Engine/Source/Runtime/WebBrowser/Private/WebBrowserHandler.h
@@ -169,6 +169,8 @@ public:
 		CefCursorHandle Cursor,
 		CefRenderHandler::CursorType Type,
 		const CefCursorInfo& CustomCursorInfo) override;
+	virtual void OnPopupShow(CefRefPtr<CefBrowser> Browser, bool bShow) override;
+	virtual void OnPopupSize(CefRefPtr<CefBrowser> Browser, const CefRect& Rect) override;
 
 public:
 
@@ -204,6 +206,12 @@ private:
 	/** Whether to show an error message in case of loading errors. */
 	bool ShowErrorMessage;
 
+	/** Size and position of the popup widget. */
+	CefRect PopupRect;
+
+	/** Set to true to show the popup widget on next popup size change. */
+	bool bShowPopupRequested;
+
 	// Include the default reference counting implementation.
 	IMPLEMENT_REFCOUNTING(FWebBrowserHandler);
 };
diff --git a/Engine/Source/Runtime/WebBrowser/Private/WebBrowserViewport.cpp b/Engine/Source/Runtime/WebBrowser/Private/WebBrowserViewport.cpp
--- a/Engine/Source/Runtime/WebBrowser/Private/WebBrowserViewport.cpp
+++ b/Engine/Source/Runtime/WebBrowser/Private/WebBrowserViewport.cpp
@@ -6,22 +6,25 @@
 
 FIntPoint FWebBrowserViewport::GetSize() const
 {
-	return (WebBrowserWindow->GetTexture() != nullptr)
-		? FIntPoint(WebBrowserWindow->GetTexture()->GetWidth(), WebBrowserWindow->GetTexture()->GetHeight())
+	return (WebBrowserWindow->GetTexture(bIsPopup) != nullptr)
+		? FIntPoint(WebBrowserWindow->GetTexture(bIsPopup)->GetWidth(), WebBrowserWindow->GetTexture(bIsPopup)->GetHeight())
 		: FIntPoint();
 }
 
 FSlateShaderResource* FWebBrowserViewport::GetViewportRenderTargetTexture() const
 {
-	return WebBrowserWindow->GetTexture();
+	return WebBrowserWindow->GetTexture(bIsPopup);
 }
 
 void FWebBrowserViewport::Tick( const FGeometry& AllottedGeometry, double InCurrentTime, float DeltaTime )
 {
-	// Calculate max corner of the viewport using same method as Slate
-	FVector2D MaxPos = AllottedGeometry.AbsolutePosition + AllottedGeometry.GetLocalSize();
-	// Get size by subtracting as int to avoid incorrect rounding when size and position are .5
-	WebBrowserWindow->SetViewportSize(MaxPos.IntPoint() - AllottedGeometry.AbsolutePosition.IntPoint());
+	if (!bIsPopup)
+	{
+		// Calculate max corner of the viewport using same method as Slate
+		FVector2D MaxPos = AllottedGeometry.AbsolutePosition + AllottedGeometry.GetLocalSize();
+		// Get size by subtracting as int to avoid incorrect rounding when size and position are .5
+		WebBrowserWindow->SetViewportSize(MaxPos.IntPoint() - AllottedGeometry.AbsolutePosition.IntPoint());
+	}
 }
 
 bool FWebBrowserViewport::RequiresVsync() const
@@ -42,7 +45,7 @@ FCursorReply FWebBrowserViewport::OnCursorQuery( const FGeometry& MyGeometry, co
 FReply FWebBrowserViewport::OnMouseButtonDown(const FGeometry& MyGeometry, const FPointerEvent& MouseEvent)
 {
 	// Capture mouse on left button down so that you can drag out of the viewport
-	FReply Reply = WebBrowserWindow->OnMouseButtonDown(MyGeometry, MouseEvent);
+	FReply Reply = WebBrowserWindow->OnMouseButtonDown(MyGeometry, MouseEvent, bIsPopup);
 	if (MouseEvent.GetEffectingButton() == EKeys::LeftMouseButton && ViewportWidget.IsValid())
 	{
 		return Reply.CaptureMouse(ViewportWidget.Pin().ToSharedRef());
@@ -53,7 +56,7 @@ FReply FWebBrowserViewport::OnMouseButtonDown(const FGeometry& MyGeometry, const
 FReply FWebBrowserViewport::OnMouseButtonUp(const FGeometry& MyGeometry, const FPointerEvent& MouseEvent)
 {
 	// Release mouse capture when left button released
-	FReply Reply = WebBrowserWindow->OnMouseButtonUp(MyGeometry, MouseEvent);
+	FReply Reply = WebBrowserWindow->OnMouseButtonUp(MyGeometry, MouseEvent, bIsPopup);
 	if (MouseEvent.GetEffectingButton() == EKeys::LeftMouseButton)
 	{
 		return Reply.ReleaseMouseCapture();
@@ -71,17 +74,17 @@ void FWebBrowserViewport::OnMouseLeave(const FPointerEvent& MouseEvent)
 
 FReply FWebBrowserViewport::OnMouseMove(const FGeometry& MyGeometry, const FPointerEvent& MouseEvent)
 {
-	return WebBrowserWindow->OnMouseMove(MyGeometry, MouseEvent);
+	return WebBrowserWindow->OnMouseMove(MyGeometry, MouseEvent, bIsPopup);
 }
 
 FReply FWebBrowserViewport::OnMouseWheel(const FGeometry& MyGeometry, const FPointerEvent& MouseEvent)
 {
-	return WebBrowserWindow->OnMouseWheel(MyGeometry, MouseEvent);
+	return WebBrowserWindow->OnMouseWheel(MyGeometry, MouseEvent, bIsPopup);
 }
 
 FReply FWebBrowserViewport::OnMouseButtonDoubleClick(const FGeometry& InMyGeometry, const FPointerEvent& InMouseEvent)
 {
-	FReply Reply = WebBrowserWindow->OnMouseButtonDoubleClick(InMyGeometry, InMouseEvent);
+	FReply Reply = WebBrowserWindow->OnMouseButtonDoubleClick(InMyGeometry, InMouseEvent, bIsPopup);
 	return Reply;
 }
 
@@ -102,11 +105,11 @@ FReply FWebBrowserViewport::OnKeyChar( const FGeometry& MyGeometry, const FChara
 
 FReply FWebBrowserViewport::OnFocusReceived(const FFocusEvent& InFocusEvent)
 {
-	WebBrowserWindow->OnFocus(true);
+	WebBrowserWindow->OnFocus(true, bIsPopup);
 	return FReply::Handled();
 }
 
 void FWebBrowserViewport::OnFocusLost(const FFocusEvent& InFocusEvent)
 {
-	WebBrowserWindow->OnFocus(false);
+	WebBrowserWindow->OnFocus(false, bIsPopup);
 }
diff --git a/Engine/Source/Runtime/WebBrowser/Private/WebBrowserWindow.cpp b/Engine/Source/Runtime/WebBrowser/Private/WebBrowserWindow.cpp
--- a/Engine/Source/Runtime/WebBrowser/Private/WebBrowserWindow.cpp
+++ b/Engine/Source/Runtime/WebBrowser/Private/WebBrowserWindow.cpp
@@ -26,7 +26,6 @@ typedef FMacCursor FPlatformCursor;
 
 FWebBrowserWindow::FWebBrowserWindow(FIntPoint InViewportSize, FString InUrl, TOptional<FString> InContentsToLoad, bool InShowErrorMessage, bool InThumbMouseButtonNavigation, bool InUseTransparency)
 	: DocumentState(EWebBrowserDocumentState::NoDocument)
-	, UpdatableTexture(nullptr)
 	, CurrentUrl(InUrl)
 	, ViewportSize(InViewportSize)
 	, bIsClosing(false)
@@ -44,19 +43,29 @@ FWebBrowserWindow::FWebBrowserWindow(FIntPoint InViewportSize, FString InUrl, TO
 	, bIgnoreKeyDownEvent(false)
 	, bIgnoreKeyUpEvent(false)
 	, bIgnoreCharacterEvent(false)
+	, bMainHasFocus(false)
+	, bPopupHasFocus(false)
 	, Scripting(new FWebJSScripting)
 {
+	UpdatableTextures[0] = nullptr;
+	UpdatableTextures[1] = nullptr;
 }
 
 FWebBrowserWindow::~FWebBrowserWindow()
 {
 	CloseBrowser(true);
-
-	if (FSlateApplication::IsInitialized() && FSlateApplication::Get().GetRenderer().IsValid() && UpdatableTexture != nullptr)
+	if(FSlateApplication::IsInitialized() && FSlateApplication::Get().GetRenderer().IsValid())
 	{
-		FSlateApplication::Get().GetRenderer()->ReleaseUpdatableTexture(UpdatableTexture);
+		for (int I = 0; I < 1; ++I)
+		{
+			if (UpdatableTextures[I] != nullptr)
+			{
+				FSlateApplication::Get().GetRenderer()->ReleaseUpdatableTexture(UpdatableTextures[I]);
+			}
+		}
 	}
-	UpdatableTexture = nullptr;
+	UpdatableTextures[0] = nullptr;
+	UpdatableTextures[1] = nullptr;
 }
 
 void FWebBrowserWindow::LoadURL(FString NewURL)
@@ -104,11 +113,11 @@ void FWebBrowserWindow::SetViewportSize(FIntPoint WindowSize)
 	}
 }
 
-FSlateShaderResource* FWebBrowserWindow::GetTexture()
+FSlateShaderResource* FWebBrowserWindow::GetTexture(bool bIsPopup)
 {
-	if (UpdatableTexture != nullptr && IsInitialized())
+	if (UpdatableTextures[bIsPopup?1:0] != nullptr && IsInitialized())
 	{
-		return UpdatableTexture->GetSlateResource();
+		return UpdatableTextures[bIsPopup?1:0]->GetSlateResource();
 	}
 	return nullptr;
 }
@@ -444,7 +453,7 @@ bool FWebBrowserWindow::RequestCreateWindow( const TSharedRef<IWebBrowserWindow>
 	return false;
 }
 
-FReply FWebBrowserWindow::OnMouseButtonDown(const FGeometry& MyGeometry, const FPointerEvent& MouseEvent)
+FReply FWebBrowserWindow::OnMouseButtonDown(const FGeometry& MyGeometry, const FPointerEvent& MouseEvent, bool bIsPopup)
 {
 	FReply Reply = FReply::Unhandled();
 	if (IsValid())
@@ -459,12 +468,7 @@ FReply FWebBrowserWindow::OnMouseButtonDown(const FGeometry& MyGeometry, const F
 			(Button == EKeys::LeftMouseButton ? MBT_LEFT : (
 			Button == EKeys::RightMouseButton ? MBT_RIGHT : MBT_MIDDLE));
 
-		CefMouseEvent Event;
-		FVector2D LocalPos = MyGeometry.AbsoluteToLocal(MouseEvent.GetScreenSpacePosition());
-		Event.x = LocalPos.X;
-		Event.y = LocalPos.Y;
-		Event.modifiers = GetCefMouseModifiers(MouseEvent);
-
+		CefMouseEvent Event = GetCefMouseEvent(MyGeometry, MouseEvent, bIsPopup);
 		InternalCefBrowser->GetHost()->SendMouseClickEvent(Event, Type, false,1);
 			Reply = FReply::Handled();
 		}
@@ -472,7 +476,7 @@ FReply FWebBrowserWindow::OnMouseButtonDown(const FGeometry& MyGeometry, const F
 	return Reply;
 }
 
-FReply FWebBrowserWindow::OnMouseButtonUp(const FGeometry& MyGeometry, const FPointerEvent& MouseEvent)
+FReply FWebBrowserWindow::OnMouseButtonUp(const FGeometry& MyGeometry, const FPointerEvent& MouseEvent, bool bIsPopup)
 {
 	FReply Reply = FReply::Unhandled();
 	if (IsValid())
@@ -487,12 +491,7 @@ FReply FWebBrowserWindow::OnMouseButtonUp(const FGeometry& MyGeometry, const FPo
 			(Button == EKeys::LeftMouseButton ? MBT_LEFT : (
 			Button == EKeys::RightMouseButton ? MBT_RIGHT : MBT_MIDDLE));
 
-		CefMouseEvent Event;
-		FVector2D LocalPos = MyGeometry.AbsoluteToLocal(MouseEvent.GetScreenSpacePosition());
-		Event.x = LocalPos.X;
-		Event.y = LocalPos.Y;
-		Event.modifiers = GetCefMouseModifiers(MouseEvent);
-
+		CefMouseEvent Event = GetCefMouseEvent(MyGeometry, MouseEvent, bIsPopup);
 		InternalCefBrowser->GetHost()->SendMouseClickEvent(Event, Type, true, 1);
 			Reply = FReply::Handled();
 		}
@@ -518,7 +517,7 @@ FReply FWebBrowserWindow::OnMouseButtonUp(const FGeometry& MyGeometry, const FPo
 	return Reply;
 }
 
-FReply FWebBrowserWindow::OnMouseButtonDoubleClick(const FGeometry& MyGeometry, const FPointerEvent& MouseEvent)
+FReply FWebBrowserWindow::OnMouseButtonDoubleClick(const FGeometry& MyGeometry, const FPointerEvent& MouseEvent, bool bIsPopup)
 {
 	FReply Reply = FReply::Unhandled();
 	if (IsValid())
@@ -533,12 +532,7 @@ FReply FWebBrowserWindow::OnMouseButtonDoubleClick(const FGeometry& MyGeometry,
 			(Button == EKeys::LeftMouseButton ? MBT_LEFT : (
 			Button == EKeys::RightMouseButton ? MBT_RIGHT : MBT_MIDDLE));
 
-		CefMouseEvent Event;
-		FVector2D LocalPos = MyGeometry.AbsoluteToLocal(MouseEvent.GetScreenSpacePosition());
-		Event.x = LocalPos.X;
-		Event.y = LocalPos.Y;
-		Event.modifiers = GetCefMouseModifiers(MouseEvent);
-
+		CefMouseEvent Event = GetCefMouseEvent(MyGeometry, MouseEvent, bIsPopup);
 		InternalCefBrowser->GetHost()->SendMouseClickEvent(Event, Type, false, 2);
 			Reply = FReply::Handled();
 	}
@@ -546,24 +540,19 @@ FReply FWebBrowserWindow::OnMouseButtonDoubleClick(const FGeometry& MyGeometry,
 	return Reply;
 }
 
-FReply FWebBrowserWindow::OnMouseMove(const FGeometry& MyGeometry, const FPointerEvent& MouseEvent)
+FReply FWebBrowserWindow::OnMouseMove(const FGeometry& MyGeometry, const FPointerEvent& MouseEvent, bool bIsPopup)
 {
 	FReply Reply = FReply::Unhandled();
 	if (IsValid())
 	{
-		CefMouseEvent Event;
-		FVector2D LocalPos = MyGeometry.AbsoluteToLocal(MouseEvent.GetScreenSpacePosition());
-		Event.x = LocalPos.X;
-		Event.y = LocalPos.Y;
-		Event.modifiers = GetCefMouseModifiers(MouseEvent);
-
+		CefMouseEvent Event = GetCefMouseEvent(MyGeometry, MouseEvent, bIsPopup);
 		InternalCefBrowser->GetHost()->SendMouseMoveEvent(Event, false);
 		Reply = FReply::Handled();
 	}
 	return Reply;
 }
 
-FReply FWebBrowserWindow::OnMouseWheel(const FGeometry& MyGeometry, const FPointerEvent& MouseEvent)
+FReply FWebBrowserWindow::OnMouseWheel(const FGeometry& MyGeometry, const FPointerEvent& MouseEvent, bool bIsPopup)
 {
 	FReply Reply = FReply::Unhandled();
 	if(IsValid())
@@ -571,12 +560,7 @@ FReply FWebBrowserWindow::OnMouseWheel(const FGeometry& MyGeometry, const FPoint
 		// The original delta is reduced so this should bring it back to what CEF expects
 		const float SpinFactor = 50.0f;
 		const float TrueDelta = MouseEvent.GetWheelDelta() * SpinFactor;
-		CefMouseEvent Event;
-		FVector2D LocalPos = MyGeometry.AbsoluteToLocal(MouseEvent.GetScreenSpacePosition());
-		Event.x = LocalPos.X;
-		Event.y = LocalPos.Y;
-		Event.modifiers = GetCefMouseModifiers(MouseEvent);
-		
+		CefMouseEvent Event = GetCefMouseEvent(MyGeometry, MouseEvent, bIsPopup);
 		InternalCefBrowser->GetHost()->SendMouseWheelEvent(Event,
 															MouseEvent.IsShiftDown() ? TrueDelta : 0,
 															!MouseEvent.IsShiftDown() ? TrueDelta : 0);
@@ -585,11 +569,21 @@ FReply FWebBrowserWindow::OnMouseWheel(const FGeometry& MyGeometry, const FPoint
 	return Reply;
 }
 
-void FWebBrowserWindow::OnFocus(bool SetFocus)
+void FWebBrowserWindow::OnFocus(bool SetFocus, bool bIsPopup)
 {
-	if (IsValid())
+	if (bIsPopup)
+	{
+		bPopupHasFocus = SetFocus;
+	}
+	else
+	{
+		bMainHasFocus = SetFocus;
+	}
+
+	// Only notify focus if there is no popup menu with focus, as SendFocusEvent will dismiss any popup menus.
+	if (IsValid() && !bPopupHasFocus)
 	{
-		InternalCefBrowser->GetHost()->SendFocusEvent(SetFocus);
+		InternalCefBrowser->GetHost()->SendFocusEvent(bMainHasFocus);
 	}
 }
 
@@ -764,19 +758,18 @@ void FWebBrowserWindow::NotifyDocumentLoadingStateChange(bool IsLoading)
 
 void FWebBrowserWindow::OnPaint(CefRenderHandler::PaintElementType Type, const CefRenderHandler::RectList& DirtyRects, const void* Buffer, int Width, int Height)
 {
-	if (UpdatableTexture == nullptr && FSlateApplication::IsInitialized() && FSlateApplication::Get().GetRenderer().IsValid())
+
+	if (UpdatableTextures[Type] == nullptr && FSlateApplication::IsInitialized() && FSlateApplication::Get().GetRenderer().IsValid())
 	{
-		UpdatableTexture = FSlateApplication::Get().GetRenderer()->CreateUpdatableTexture(Width,Height);
+		UpdatableTextures[Type] = FSlateApplication::Get().GetRenderer()->CreateUpdatableTexture(Width,Height);
 	}
 
-	if (UpdatableTexture != nullptr)
+	if (UpdatableTextures[Type] != nullptr)
 	{
-		int32 FirstRow = 0;
-		int32 LastRow = Height;
 		// Note that with more recent versions of CEF, the DirtyRects will always contain a single element, as it merges all dirty areas into a single rectangle before calling OnPaint
 		// In case that should change in the future, we'll simply update the entire area if DirtyRects is not a single element.
 		FIntRect Dirty = (DirtyRects.size() == 1)?FIntRect(DirtyRects[0].x, DirtyRects[0].y, DirtyRects[0].x + DirtyRects[0].width, DirtyRects[0].y + DirtyRects[0].height):FIntRect();
-		UpdatableTexture->UpdateTextureThreadSafeRaw(Width, Height, Buffer, Dirty);
+		UpdatableTextures[Type]->UpdateTextureThreadSafeRaw(Width, Height, Buffer, Dirty);
 	}
 
 	bIsInitialized = true;
@@ -968,6 +961,21 @@ int32 FWebBrowserWindow::GetCefMouseModifiers(const FPointerEvent& InMouseEvent)
 	return Modifiers;
 }
 
+CefMouseEvent FWebBrowserWindow::GetCefMouseEvent(const FGeometry& MyGeometry, const FPointerEvent& MouseEvent, bool bIsPopup)
+{
+	CefMouseEvent Event;
+	FVector2D LocalPos = MyGeometry.AbsoluteToLocal(MouseEvent.GetScreenSpacePosition());
+	Event.x = LocalPos.X;
+	Event.y = LocalPos.Y;
+	if (bIsPopup)
+	{
+		Event.x += PopupRect.Min.X;
+		Event.y += PopupRect.Min.Y;
+	}
+	Event.modifiers = GetCefMouseModifiers(MouseEvent);
+	return Event;
+}
+
 int32 FWebBrowserWindow::GetCefInputModifiers(const FInputEvent& InputEvent)
 {
 	int32 Modifiers = 0;
@@ -1099,4 +1107,18 @@ void FWebBrowserWindow::OnBrowserClosed()
 	Handler = nullptr;
 }
 
+void FWebBrowserWindow::ShowPopup(CefRect CefPopupSize)
+{
+	PopupRect = FIntRect(CefPopupSize.x, CefPopupSize.y, CefPopupSize.x+CefPopupSize.width, CefPopupSize.y+CefPopupSize.height);
+	bPopupHasFocus = true;
+	OnShowPopup().Broadcast(PopupRect);
+}
+
+void FWebBrowserWindow::HidePopup()
+{
+	bPopupHasFocus = false;
+	OnDismissPopup().Broadcast();
+}
+
+
 #endif
diff --git a/Engine/Source/Runtime/WebBrowser/Private/WebBrowserWindow.h b/Engine/Source/Runtime/WebBrowser/Private/WebBrowserWindow.h
--- a/Engine/Source/Runtime/WebBrowser/Private/WebBrowserWindow.h
+++ b/Engine/Source/Runtime/WebBrowser/Private/WebBrowserWindow.h
@@ -90,7 +90,7 @@ public:
 	virtual void LoadURL(FString NewURL) override;
 	virtual void LoadString(FString Contents, FString DummyURL) override;
 	virtual void SetViewportSize(FIntPoint WindowSize) override;
-	virtual FSlateShaderResource* GetTexture() override;
+	virtual FSlateShaderResource* GetTexture(bool bIsPopup = false) override;
 	virtual bool IsValid() const override;
 	virtual bool IsInitialized() const override;
 	virtual bool IsClosing() const override;
@@ -100,12 +100,12 @@ public:
 	virtual bool OnKeyDown(const FKeyEvent& InKeyEvent) override;
 	virtual bool OnKeyUp(const FKeyEvent& InKeyEvent) override;
 	virtual bool OnKeyChar(const FCharacterEvent& InCharacterEvent) override;
-	virtual FReply OnMouseButtonDown(const FGeometry& MyGeometry, const FPointerEvent& MouseEvent) override;
-	virtual FReply OnMouseButtonUp(const FGeometry& MyGeometry, const FPointerEvent& MouseEvent) override;
-	virtual FReply OnMouseButtonDoubleClick(const FGeometry& MyGeometry, const FPointerEvent& MouseEvent) override;
-	virtual FReply OnMouseMove(const FGeometry& MyGeometry, const FPointerEvent& MouseEvent) override;
-	virtual FReply OnMouseWheel(const FGeometry& MyGeometry, const FPointerEvent& MouseEvent) override;
-	virtual void OnFocus(bool SetFocus) override;
+	virtual FReply OnMouseButtonDown(const FGeometry& MyGeometry, const FPointerEvent& MouseEvent, bool bIsPopup) override;
+	virtual FReply OnMouseButtonUp(const FGeometry& MyGeometry, const FPointerEvent& MouseEvent, bool bIsPopup) override;
+	virtual FReply OnMouseButtonDoubleClick(const FGeometry& MyGeometry, const FPointerEvent& MouseEvent, bool bIsPopup) override;
+	virtual FReply OnMouseMove(const FGeometry& MyGeometry, const FPointerEvent& MouseEvent, bool bIsPopup) override;
+	virtual FReply OnMouseWheel(const FGeometry& MyGeometry, const FPointerEvent& MouseEvent, bool bIsPopup) override;
+	virtual void OnFocus(bool SetFocus, bool bIsPopup) override;
 	virtual void OnCaptureLost() override;
 	virtual bool CanGoBack() const override;
 	virtual void GoBack() override;
@@ -179,6 +179,18 @@ public:
 		return BeforePopupDelegate;
 	}
 
+	DECLARE_DERIVED_EVENT(FWebBrowserWindow, IWebBrowserWindow::FOnShowPopup, FOnShowPopup);
+	virtual FOnShowPopup& OnShowPopup() override
+	{
+		return ShowPopupEvent;
+	}
+
+	DECLARE_DERIVED_EVENT(FWebBrowserWindow, IWebBrowserWindow::FOnDismissPopup, FOnDismissPopup);
+	virtual FOnDismissPopup& OnDismissPopup() override
+	{
+		return DismissPopupEvent;
+	}
+
 private:
 
 	/**
@@ -305,6 +317,18 @@ private:
 	 */
 	void OnBrowserClosed();
 
+	/**
+	 * Called to show the popup widget
+	 *
+ 	 * @param PopupSize The size and location of the popup widget.
+	 */
+	void ShowPopup(CefRect PopupSize);
+
+	/**
+	 * Called to hide the popup widget again
+	 */
+	void HidePopup();
+
 public:
 
 	// Trigger an OnBeforePopup event chain
@@ -354,13 +378,16 @@ private:
 	/** Used by the key down and up handlers to convert Slate key events to the CEF equivalent. */
 	void PopulateCefKeyEvent(const FKeyEvent& InKeyEvent, CefKeyEvent& OutKeyEvent);
 
+	/** Used to convert a FPointerEvent to a CefMouseEvent */
+	CefMouseEvent GetCefMouseEvent(const FGeometry& MyGeometry, const FPointerEvent& MouseEvent, bool bIsPopup);
+
 private:
 
 	/** Current state of the document being loaded. */
 	EWebBrowserDocumentState DocumentState;
 
 	/** Interface to the texture we are rendering to. */
-	FSlateUpdatableTexture* UpdatableTexture;
+	FSlateUpdatableTexture* UpdatableTextures[2];
 
 	/** Pointer to the CEF Handler for this window. */
 	CefRefPtr<FWebBrowserHandler> Handler;
@@ -428,6 +455,12 @@ private:
 	/** Delegate for handaling requests to close new windows that were created. */
 	FOnCloseWindow CloseWindowDelegate;
 
+	/** Delegate for handaling requests to show the popup menu. */
+	FOnShowPopup ShowPopupEvent;
+
+	/** Delegate for handaling requests to close new windows that were created. */
+	FOnDismissPopup DismissPopupEvent;
+
 	/** Tracks the current mouse cursor */
 	EMouseCursor::Type Cursor;
 
@@ -445,6 +478,12 @@ private:
 	bool bIgnoreKeyUpEvent;
 	bool bIgnoreCharacterEvent;
 
+	/** Used to ignore any popup menus when forwarding focus gained/lost events*/
+	bool bMainHasFocus;
+	bool bPopupHasFocus;
+
+	FIntRect PopupRect;
+
 	/** Handling of passing and marshalling messages for JS integration is delegated to a helper class*/
 	TSharedPtr<FWebJSScripting> Scripting;
 };
diff --git a/Engine/Source/Runtime/WebBrowser/Public/IWebBrowserWindow.h b/Engine/Source/Runtime/WebBrowser/Public/IWebBrowserWindow.h
--- a/Engine/Source/Runtime/WebBrowser/Public/IWebBrowserWindow.h
+++ b/Engine/Source/Runtime/WebBrowser/Public/IWebBrowserWindow.h
@@ -52,9 +52,10 @@ public:
 	/**
 	 * Gets interface to the texture representation of the browser
 	 *
+	 * @param bISpopup Whether to return the popup menu texture instead of the main browser window.
 	 * @return A slate shader resource that can be rendered
 	 */
-	virtual FSlateShaderResource* GetTexture() = 0;
+	virtual FSlateShaderResource* GetTexture(bool bIsPopup = false) = 0;
 
 	/**
 	 * Checks whether the web browser is valid and ready for use
@@ -112,50 +113,55 @@ public:
 	 *
 	 * @param MyGeometry The Geometry of the browser
 	 * @param MouseEvent Information about the input event
-	 * 
+	 * @param bIsPopup True if the coordinates are relative to a popup menu window, otherwise false.
+	 *
 	 * @return FReply::Handled() if the mouse event was handled, FReply::Unhandled() oterwise
 	 */
-	virtual FReply OnMouseButtonDown(const FGeometry& MyGeometry, const FPointerEvent& MouseEvent) = 0;
+	virtual FReply OnMouseButtonDown(const FGeometry& MyGeometry, const FPointerEvent& MouseEvent, bool bIsPopup) = 0;
 
 	/**
 	 * Notify the browser that a mouse button was released within it
 	 *
 	 * @param MyGeometry The Geometry of the browser
 	 * @param MouseEvent Information about the input event
-	 * 
+	 * @param bIsPopup True if the coordinates are relative to a popup menu window, otherwise false.
+	 *
 	 * @return FReply::Handled() if the mouse event was handled, FReply::Unhandled() oterwise
 	 */
-	virtual FReply OnMouseButtonUp(const FGeometry& MyGeometry, const FPointerEvent& MouseEvent) = 0;
+	virtual FReply OnMouseButtonUp(const FGeometry& MyGeometry, const FPointerEvent& MouseEvent, bool bIsPopup) = 0;
 
 	/**
 	 * Notify the browser of a double click event
 	 *
 	 * @param MyGeometry The Geometry of the browser
 	 * @param MouseEvent Information about the input event
-	 * 
+	 * @param bIsPopup True if the coordinates are relative to a popup menu window, otherwise false.
+	 *
 	 * @return FReply::Handled() if the mouse event was handled, FReply::Unhandled() oterwise
 	 */
-	virtual FReply OnMouseButtonDoubleClick(const FGeometry& MyGeometry, const FPointerEvent& MouseEvent) = 0;
+	virtual FReply OnMouseButtonDoubleClick(const FGeometry& MyGeometry, const FPointerEvent& MouseEvent, bool bIsPopup) = 0;
 
 	/**
 	 * Notify the browser that a mouse moved within it
 	 *
 	 * @param MyGeometry The Geometry of the browser
 	 * @param MouseEvent Information about the input event
-	 * 
+	 * @param bIsPopup True if the coordinates are relative to a popup menu window, otherwise false.
+	 *
 	 * @return FReply::Handled() if the mouse event was handled, FReply::Unhandled() oterwise
 	 */
-	virtual FReply OnMouseMove(const FGeometry& MyGeometry, const FPointerEvent& MouseEvent) = 0;
+	virtual FReply OnMouseMove(const FGeometry& MyGeometry, const FPointerEvent& MouseEvent, bool bIsPopup) = 0;
 
 	/**
 	 * Called when the mouse wheel is spun
 	 *
 	 * @param MyGeometry The Geometry of the browser
 	 * @param MouseEvent Information about the input event
-	 * 
+	 * @param bIsPopup True if the coordinates are relative to a popup menu window, otherwise false.
+	 *
 	 * @return FReply::Handled() if the mouse event was handled, FReply::Unhandled() oterwise
 	 */
-	virtual FReply OnMouseWheel(const FGeometry& MyGeometry, const FPointerEvent& MouseEvent) = 0;
+	virtual FReply OnMouseWheel(const FGeometry& MyGeometry, const FPointerEvent& MouseEvent, bool bIsPopup) = 0;
 
 	/**
 	 * The system asks each widget under the mouse to provide a cursor. This event is bubbled.
@@ -166,8 +172,10 @@ public:
 
 	/**
 	 * Called when browser receives/loses focus
+	 * @param SetFocus Whether the window gained or lost focus.
+	 * @param bIsPopup True if the coordinates are relative to a popup menu window, otherwise false.
 	 */
-	virtual void OnFocus(bool SetFocus) = 0;
+	virtual void OnFocus(bool SetFocus, bool bIsPopup) = 0;
 
 	/** Called when Capture lost */
 	virtual void OnCaptureLost() = 0;
@@ -271,6 +279,14 @@ public:
 	DECLARE_DELEGATE_RetVal_OneParam(bool, FOnCloseWindow, const TWeakPtr<IWebBrowserWindow>& /*BrowserWindow*/)
 	virtual FOnCloseWindow& OnCloseWindow() = 0;
 
+	/** A delegate that is invoked when the browser needs to show a popup menu. */
+	DECLARE_EVENT_OneParam(IWebBrowserWindow, FOnShowPopup, const FIntRect& /*PopupSize*/)
+	virtual FOnShowPopup& OnShowPopup() = 0;
+
+	/** A delegate that is invoked when the browser no longer wants to show the popup menu. */
+	DECLARE_EVENT(IWebBrowserWindow, FOnDismissPopup)
+	virtual FOnDismissPopup& OnDismissPopup() = 0;
+
 protected:
 
 	/** Virtual Destructor. */
diff --git a/Engine/Source/Runtime/WebBrowser/Public/SWebBrowser.h b/Engine/Source/Runtime/WebBrowser/Public/SWebBrowser.h
--- a/Engine/Source/Runtime/WebBrowser/Public/SWebBrowser.h
+++ b/Engine/Source/Runtime/WebBrowser/Public/SWebBrowser.h
@@ -249,6 +249,12 @@ private:
 	/** Callback for popup window permission */
 	bool HandleBeforePopup(FString URL, FString Target);
 
+	/** Callback for showing a popup menu */
+	void HandleShowPopup(const FIntRect& PopupSize);
+
+	/** Callback for hiding the popup menu */
+	void HandleDismissPopup();
+
 private:
 
 	/** Interface for dealing with a web browser window. */
@@ -256,10 +262,18 @@ private:
 
 	/** Viewport interface for rendering the web page. */
 	TSharedPtr<FWebBrowserViewport> BrowserViewport;
+	/** Viewport interface for rendering popup menus. */
+	TSharedPtr<FWebBrowserViewport>	MenuViewport;
 
 	/** The actual viewport widget. Required to update its tool tip property. */
 	TSharedPtr<SViewport> ViewportWidget;
 
+	/**
+	 * An interface pointer to a menu object presenting a popup.
+	 * Pointer is null when a popup is not visible.
+	 */
+	TWeakPtr<IMenu> PopupMenuPtr;
+
 	/** The url that appears in the address bar which can differ from the url of the loaded page */
 	FText AddressBarUrl;
 
diff --git a/Engine/Source/Runtime/WebBrowser/Public/WebBrowserViewport.h b/Engine/Source/Runtime/WebBrowser/Public/WebBrowserViewport.h
--- a/Engine/Source/Runtime/WebBrowser/Public/WebBrowserViewport.h
+++ b/Engine/Source/Runtime/WebBrowser/Public/WebBrowserViewport.h
@@ -18,10 +18,12 @@ public:
 	 * 
 	 * @param InWebBrowserWindow The Web Browser Window this viewport will display
 	 * @param InViewportWidget The Widget displaying this viewport (needed to capture mouse)
+	 * @param InIsPopup Used to initialize a viewport for showing browser popup menus instead of the main window.
 	 */
-	FWebBrowserViewport(TSharedPtr<IWebBrowserWindow> InWebBrowserWindow, TSharedPtr<SWidget> InViewportWidget)
+	FWebBrowserViewport(TSharedPtr<IWebBrowserWindow> InWebBrowserWindow, TSharedPtr<SWidget> InViewportWidget, bool InIsPopup = false)
 		: WebBrowserWindow(InWebBrowserWindow)
 		, ViewportWidget(InViewportWidget)
+		, bIsPopup(InIsPopup)
 	{ }
 
 	/**
@@ -56,4 +58,6 @@ private:
 	TSharedPtr<IWebBrowserWindow>	WebBrowserWindow;
 	/** The viewport widget using this interface */
 	TWeakPtr<SWidget>				ViewportWidget;
+	/** Whether this viewport is showing the browser window or a popup menu widget */
+	bool const						bIsPopup;
 };
