[INTEGRATE] Merging CL#2613497, CL#2614045, and CL#2614288 using UE4-To-UE4-LauncherDev:
Changes to fix background color issues when opening Chrome Dev Tools from a SBrowserWindow with transparency enabled.

diff --git a/Engine/Source/Runtime/WebBrowser/Private/WebBrowserHandler.cpp b/Engine/Source/Runtime/WebBrowser/Private/WebBrowserHandler.cpp
--- a/Engine/Source/Runtime/WebBrowser/Private/WebBrowserHandler.cpp
+++ b/Engine/Source/Runtime/WebBrowser/Private/WebBrowserHandler.cpp
@@ -144,8 +144,13 @@ bool FWebBrowserHandler::OnBeforePopup( CefRefPtr<CefBrowser> Browser,
 				
 				CefWindowHandle ParentWindowHandle = BrowserWindow->GetCefBrowser()->GetHost()->GetWindowHandle();
 
+				// Allow overriding transparency setting for child windows
+				bool bUseTransparency = BrowserWindow->UseTransparency()
+										? NewHandler->BrowserPopupFeatures->GetAdditionalFeatures().Find(TEXT("Epic_NoTransparency")) == INDEX_NONE
+										: NewHandler->BrowserPopupFeatures->GetAdditionalFeatures().Find(TEXT("Epic_UseTransparency")) != INDEX_NONE;
+
 				// Always use off screen rendering so we can integrate with our windows
-				WindowInfo.SetAsWindowless(ParentWindowHandle, BrowserWindow->UseTransparency());
+				WindowInfo.SetAsWindowless(ParentWindowHandle, bUseTransparency);
 
 				// We need to rely on CEF to create our window so we set the WindowInfo, BrowserSettings, Client, and then return false
 				bSupressCEFWindowCreation = false;
@@ -353,6 +358,49 @@ bool FWebBrowserHandler::OnProcessMessageReceived(CefRefPtr<CefBrowser> Browser,
 	return Retval;
 }
 
+bool FWebBrowserHandler::ShowDevTools(const CefRefPtr<CefBrowser>& Browser)
+{
+	CefPoint Point;
+	CefString TargetUrl = "chrome-devtools://devtools/devtools.html";
+	CefString TargetFrameName = "devtools";
+	CefPopupFeatures PopupFeatures;
+	CefWindowInfo WindowInfo;
+	CefRefPtr<CefClient> NewClient;
+	CefBrowserSettings BrowserSettings;
+	bool NoJavascriptAccess = false;
+
+	PopupFeatures.xSet = false;
+	PopupFeatures.ySet = false;
+	PopupFeatures.heightSet = false;
+	PopupFeatures.widthSet = false;
+	PopupFeatures.locationBarVisible = false;
+	PopupFeatures.menuBarVisible = false;
+	PopupFeatures.toolBarVisible  = false;
+	PopupFeatures.statusBarVisible  = false;
+	PopupFeatures.resizable = true;
+	PopupFeatures.additionalFeatures = cef_string_list_alloc();
+
+	// Override the parent window setting for transparency, as the Dev Tools require a solid background
+	CefString OverrideTransparencyFeature = "Epic_NoTransparency";
+	cef_string_list_append(PopupFeatures.additionalFeatures, OverrideTransparencyFeature.GetStruct());
+
+	// Set max framerate to maximum supported.
+	BrowserSettings.windowless_frame_rate = 60;
+	// Disable plugins
+	BrowserSettings.plugins = STATE_DISABLED;
+	// Dev Tools look best with a white background color
+	BrowserSettings.background_color = CefColorSetARGB(255, 255, 255, 255);
+
+	// OnBeforePopup already takes care of all the details required to ask the host application to create a new browser window.
+	bool bSuppressWindowCreation = OnBeforePopup(Browser, Browser->GetFocusedFrame(), TargetUrl, TargetFrameName, PopupFeatures, WindowInfo, NewClient, BrowserSettings, &NoJavascriptAccess);
+
+	if(! bSuppressWindowCreation)
+	{
+		Browser->GetHost()->ShowDevTools(WindowInfo, NewClient, BrowserSettings, Point);
+	}
+	return !bSuppressWindowCreation;
+}
+
 bool FWebBrowserHandler::OnKeyEvent(CefRefPtr<CefBrowser> Browser,
 	const CefKeyEvent& Event,
 	CefEventHandle OsEvent)
@@ -368,38 +416,7 @@ bool FWebBrowserHandler::OnKeyEvent(CefRefPtr<CefBrowser> Browser,
 		(Event.unmodified_character == 'i' || Event.unmodified_character == 'I')
 	  )
 	{
-
-		CefPoint Point;
-		CefString TargetUrl = "chrome-devtools://devtools/devtools.html";
-		CefString TargetFrameName = "devtools";
-		CefPopupFeatures PopupFeatures;
-		CefWindowInfo WindowInfo;
-		CefRefPtr<CefClient> NewClient;
-		CefBrowserSettings BrowserSettings;
-		bool NoJavascriptAccess = false;
-
-		PopupFeatures.xSet = false;
-		PopupFeatures.ySet = false;
-		PopupFeatures.heightSet = false;
-		PopupFeatures.widthSet = false;
-		PopupFeatures.locationBarVisible = false;
-		PopupFeatures.menuBarVisible = false;
-		PopupFeatures.toolBarVisible  = false;
-		PopupFeatures.statusBarVisible  = false;
-		PopupFeatures.resizable = true;
-
-		// Set max framerate to maximum supported.
-		BrowserSettings.windowless_frame_rate = 60;
-		// Disable plugins
-		BrowserSettings.plugins = STATE_DISABLED;
-
-		// OnBeforePopup already takes care of all the details required to ask the host application to create a new browser window.
-		bool bSuppressWindowCreation = OnBeforePopup(Browser, Browser->GetFocusedFrame(), TargetUrl, TargetFrameName, PopupFeatures, WindowInfo, NewClient, BrowserSettings, &NoJavascriptAccess);
-
-		if(! bSuppressWindowCreation)
-		{
-			Browser->GetHost()->ShowDevTools(WindowInfo, NewClient, BrowserSettings, Point);
-		}
+		return ShowDevTools(Browser);
 	}
 #endif
 
diff --git a/Engine/Source/Runtime/WebBrowser/Private/WebBrowserHandler.h b/Engine/Source/Runtime/WebBrowser/Private/WebBrowserHandler.h
--- a/Engine/Source/Runtime/WebBrowser/Private/WebBrowserHandler.h
+++ b/Engine/Source/Runtime/WebBrowser/Private/WebBrowserHandler.h
@@ -194,6 +194,8 @@ public:
 
 private:
 
+	bool ShowDevTools(const CefRefPtr<CefBrowser>& Browser);
+
 	/** Weak Pointer to our Web Browser window so that events can be passed on while it's valid*/
 	TWeakPtr<FWebBrowserWindow> BrowserWindowPtr;
 	
