[INTEGRATE] Merging CL#2615423, CL#2616952, CL#2616995 using UE4-To-UE4-LauncherDev:
Various JS integration fixes and cleanup.

diff --git a/Engine/Source/Runtime/WebBrowser/Private/WebJSScripting.cpp b/Engine/Source/Runtime/WebBrowser/Private/WebJSScripting.cpp
--- a/Engine/Source/Runtime/WebBrowser/Private/WebJSScripting.cpp
+++ b/Engine/Source/Runtime/WebBrowser/Private/WebJSScripting.cpp
@@ -142,7 +142,7 @@ FWebJSParam FWebJSScripting::ConvertResult(UProperty* Property, uint8* Data)
 }
 
 
-CefRefPtr<CefDictionaryValue> FWebJSScripting::ConvertStruct(UStruct* TypeInfo, void* StructPtr)
+CefRefPtr<CefDictionaryValue> FWebJSScripting::ConvertStruct(UStruct* TypeInfo, const void* StructPtr)
 {
 	FWebJSStructSerializerBackend Backend (SharedThis(this));
 	FStructSerializer::Serialize(StructPtr, *TypeInfo, Backend);
diff --git a/Engine/Source/Runtime/WebBrowser/Private/WebJSScripting.h b/Engine/Source/Runtime/WebBrowser/Private/WebJSScripting.h
--- a/Engine/Source/Runtime/WebBrowser/Private/WebJSScripting.h
+++ b/Engine/Source/Runtime/WebBrowser/Private/WebJSScripting.h
@@ -48,7 +48,7 @@ public:
 	 */
 	void SendProcessMessage(CefRefPtr<CefProcessMessage> Message);
 
-	CefRefPtr<CefDictionaryValue> ConvertStruct(UStruct* TypeInfo, void* StructPtr);
+	CefRefPtr<CefDictionaryValue> ConvertStruct(UStruct* TypeInfo, const void* StructPtr);
 	CefRefPtr<CefDictionaryValue> ConvertObject(UObject* Object);
 	CefRefPtr<CefDictionaryValue> ConvertObject(const FGuid& Key);
 
@@ -73,8 +73,15 @@ public:
 			}
 			case FWebJSParam::PTYPE_OBJECT:
 			{
-				CefRefPtr<CefDictionaryValue> ConvertedObject = ConvertObject(Param.ObjectValue);
-				return Container->SetDictionary(Key, ConvertedObject);
+				if (Param.ObjectValue == nullptr)
+				{
+					return Container->SetNull(Key);
+				}
+				else
+				{
+					CefRefPtr<CefDictionaryValue> ConvertedObject = ConvertObject(Param.ObjectValue);
+					return Container->SetDictionary(Key, ConvertedObject);
+				}
 			}
 			case FWebJSParam::PTYPE_STRUCT:
 			{
diff --git a/Engine/Source/Runtime/WebBrowser/Public/WebJSFunction.h b/Engine/Source/Runtime/WebBrowser/Public/WebJSFunction.h
--- a/Engine/Source/Runtime/WebBrowser/Public/WebJSFunction.h
+++ b/Engine/Source/Runtime/WebBrowser/Public/WebJSFunction.h
@@ -15,33 +15,35 @@ struct WEBBROWSER_API FWebJSParam
 		double DoubleValue;
 		int32 IntValue;
 		UObject* ObjectValue;
-		FString* StringValue;
+		const FString* StringValue;
 		struct {
 			UStruct* TypeInfo;
-			void* StructPtr;
+			const void* StructPtr;
 		} StructValue;
 	};
 
 	FWebJSParam() : Tag(PTYPE_NULL) {}
 	FWebJSParam(bool Value) : Tag(PTYPE_BOOL), BoolValue(Value) {}
+	FWebJSParam(int8 Value) : Tag(PTYPE_INT), IntValue(Value) {}
+	FWebJSParam(int16 Value) : Tag(PTYPE_INT), IntValue(Value) {}
 	FWebJSParam(int32 Value) : Tag(PTYPE_INT), IntValue(Value) {}
+	FWebJSParam(uint8 Value) : Tag(PTYPE_INT), IntValue(Value) {}
+	FWebJSParam(uint16 Value) : Tag(PTYPE_INT), IntValue(Value) {}
+	FWebJSParam(uint32 Value) : Tag(PTYPE_DOUBLE), DoubleValue(Value) {}
+	FWebJSParam(int64 Value) : Tag(PTYPE_DOUBLE), DoubleValue(Value) {}
+	FWebJSParam(uint64 Value) : Tag(PTYPE_DOUBLE), DoubleValue(Value) {}
 	FWebJSParam(double Value) : Tag(PTYPE_DOUBLE), DoubleValue(Value) {}
+	FWebJSParam(float Value) : Tag(PTYPE_DOUBLE), DoubleValue(Value) {}
 	FWebJSParam(const FString& Value) : Tag(PTYPE_STRING), StringValue(new FString(Value)) {}
 	FWebJSParam(const TCHAR* Value) : Tag(PTYPE_STRING), StringValue(new FString(Value)) {}
 	FWebJSParam(UObject* Value) : Tag(PTYPE_OBJECT), ObjectValue(Value) {}
-	FWebJSParam(UStruct* Struct, void* Value) : Tag(PTYPE_STRUCT), StructValue({Struct, Value}) {}
+	FWebJSParam(UStruct* Struct, const void* Value) : Tag(PTYPE_STRUCT), StructValue({Struct, Value}) {}
+	template <typename T> FWebJSParam(const T& Value) : Tag(PTYPE_STRUCT), StructValue({T::StaticStruct(), &Value}) {}
 	FWebJSParam(const FWebJSParam& Other);
 	~FWebJSParam();
 
 };
 
-/** Add this macro to a USTRUCT definition to be able to pass instances to FWebJSFunctions */
-#define DECLARE_JSPARAM_STRUCT() \
-	explicit operator FWebJSParam() \
-	{ \
-		return FWebJSParam(StaticStruct(), (void*)this);\
-	}
-
 USTRUCT()
 struct WEBBROWSER_API FWebJSFunction
 {
