[INTEGRATE] Use new serializer predicate functionality to serialize just the return value from the arguments struct.
Reduces amount of required code and adds support to return arrays from UFunctions to JS code.
Merging CL#2624614 using UE4-To-UE4-LauncherDev

diff --git a/Engine/Source/Runtime/WebBrowser/Private/WebJSScripting.cpp b/Engine/Source/Runtime/WebBrowser/Private/WebJSScripting.cpp
--- a/Engine/Source/Runtime/WebBrowser/Private/WebJSScripting.cpp
+++ b/Engine/Source/Runtime/WebBrowser/Private/WebJSScripting.cpp
@@ -45,103 +45,6 @@ namespace
 	}
 }
 
-FWebJSParam FWebJSScripting::ConvertResult(UProperty* Property, uint8* Data)
-{
-	// booleans
-	if (Property->IsA<UBoolProperty>())
-	{
-		return FWebJSParam(Cast<UBoolProperty>(Property)->GetPropertyValue_InContainer(Data));
-	}
-
-	// unsigned bytes & enumerations
-	else if (Property->IsA<UByteProperty>())
-	{
-		UByteProperty* ByteProperty = Cast<UByteProperty>(Property);
-
-		if (ByteProperty->IsEnum())
-		{
-			return FWebJSParam(ByteProperty->Enum->GetEnumName(ByteProperty->GetPropertyValue_InContainer(Data)));
-		}
-		else
-		{
-			return FWebJSParam((int32)ByteProperty->GetPropertyValue_InContainer(Data));
-		}
-	}
-
-	// floating point numbers
-	else if (Property->IsA<UDoubleProperty>())
-	{
-		return FWebJSParam(Cast<UDoubleProperty>(Property)->GetPropertyValue_InContainer(Data));
-	}
-	else if (Property->IsA<UFloatProperty>())
-	{
-		return FWebJSParam(Cast<UFloatProperty>(Property)->GetPropertyValue_InContainer(Data));
-	}
-
-	// signed integers
-	else if (Property->IsA<UIntProperty>())
-	{
-		return FWebJSParam(Cast<UIntProperty>(Property)->GetPropertyValue_InContainer(Data));
-	}
-	else if (Property->IsA<UInt8Property>())
-	{
-		return FWebJSParam((int32)Cast<UInt8Property>(Property)->GetPropertyValue_InContainer(Data));
-	}
-	else if (Property->IsA<UInt16Property>())
-	{
-		return FWebJSParam((int32)Cast<UInt16Property>(Property)->GetPropertyValue_InContainer(Data));
-	}
-	else if (Property->IsA<UInt64Property>())
-	{
-		return FWebJSParam((double)Cast<UInt64Property>(Property)->GetPropertyValue_InContainer(Data));
-	}
-
-	// unsigned integers
-	else if (Property->IsA<UUInt16Property>())
-	{
-		return FWebJSParam((int32)Cast<UUInt16Property>(Property)->GetPropertyValue_InContainer(Data));
-	}
-	else if (Property->IsA<UUInt32Property>())
-	{
-		return FWebJSParam((double)Cast<UUInt32Property>(Property)->GetPropertyValue_InContainer(Data));
-	}
-	else if (Property->IsA<UUInt64Property>())
-	{
-		return FWebJSParam((double)Cast<UUInt64Property>(Property)->GetPropertyValue_InContainer(Data));
-	}
-
-	// names & strings
-	else if (Property->IsA<UNameProperty>())
-	{
-		return FWebJSParam(Cast<UNameProperty>(Property)->GetPropertyValue_InContainer(Data).ToString());
-	}
-	else if (Property->IsA<UStrProperty>())
-	{
-		return FWebJSParam(Cast<UStrProperty>(Property)->GetPropertyValue_InContainer(Data));
-	}
-
-	// classes & objects
-	else if (Property->IsA<UClassProperty>())
-	{
-		return FWebJSParam(Cast<UClassProperty>(Property)->GetPropertyValue_InContainer(Data)->GetPathName());
-	}
-	else if (Property->IsA<UObjectProperty>())
-	{
-		return FWebJSParam(Cast<UObjectProperty>(Property)->GetPropertyValue_InContainer(Data));
-	}
-	else if (Property->IsA<UStructProperty>())
-	{
-		UStructProperty* StructProperty = Cast<UStructProperty>(Property);
-		return FWebJSParam(StructProperty->Struct, StructProperty->ContainerPtrToValuePtr<void>(Data));
-	}
-	else
-	{
-		GLog->Logf(ELogVerbosity::Warning, TEXT("FWebJSScripting: %s cannot be serialized, because its type (%s) is not supported"), *Property->GetName(), *Property->GetClass()->GetName());
-	}
-	return FWebJSParam();
-}
-
-
 CefRefPtr<CefDictionaryValue> FWebJSScripting::ConvertStruct(UStruct* TypeInfo, const void* StructPtr)
 {
 	FWebJSStructSerializerBackend Backend (SharedThis(this));
@@ -367,25 +270,31 @@ bool FWebJSScripting::HandleExecuteUObjectMethodMessage(CefRefPtr<CefListValue>
 		if (ParamsSize > 0)
 		{
 			// UFunction is a subclass of UStruct, so we can treat the arguments as a struct for deserialization
-			UStruct* TypeInfo = Cast<UStruct>(Function);
 
 			Params.AddUninitialized(ParamsSize);
-			TypeInfo->InitializeStruct(Params.GetData());
+			Function->InitializeStruct(Params.GetData());
 			FWebJSStructDeserializerBackend Backend = FWebJSStructDeserializerBackend(SharedThis(this), NamedArgs);
-			FStructDeserializer::Deserialize(Params.GetData(), *TypeInfo, Backend);
+			FStructDeserializer::Deserialize(Params.GetData(), *Function, Backend);
 		}
 	}
 
 	Object->ProcessEvent(Function, Params.GetData());
+	CefRefPtr<CefListValue> Results = CefListValue::Create();
 	if ( ReturnParam )
 	{
-		FWebJSParam Results[1] = {ConvertResult(ReturnParam, Params.GetData())};
-		InvokeJSFunction(ResultCallbackId, 1, Results, false);
-	}
-	else
-	{
-		InvokeJSFunction(ResultCallbackId, 0, nullptr, false);
+		FStructSerializerPolicies ReturnPolicies;
+		ReturnPolicies.PropertyFilter = [&](const UProperty* CandidateProperty, const UProperty* ParentProperty)
+		{
+			return ParentProperty != nullptr || CandidateProperty == ReturnParam;
+		};
+		FWebJSStructSerializerBackend ReturnBackend(SharedThis(this));
+		FStructSerializer::Serialize(Params.GetData(), *Function, ReturnBackend, ReturnPolicies);
+		CefRefPtr<CefDictionaryValue> ResultDict = ReturnBackend.GetResult();
+
+		// Extract the single return value from the serialized dictionary to an array
+		CopyContainerValue(Results, ResultDict, 0, CefString(*ReturnParam->GetName()));
 	}
+	InvokeJSFunction(ResultCallbackId, Results, false);
 	return true;
 }
 
@@ -400,7 +309,6 @@ void FWebJSScripting::UnbindCefBrowser()
 	InternalCefBrowser = nullptr;
 }
 
-
 void FWebJSScripting::AddReferencedObjects(FReferenceCollector& Collector)
 {
 	// Ensure bound UObjects are not garbage collected as long as this object is valid.
@@ -412,8 +320,9 @@ void FWebJSScripting::AddReferencedObjects(FReferenceCollector& Collector)
 
 void FWebJSScripting::InvokeJSErrorResult(FGuid FunctionId, const FString& Error)
 {
-	FWebJSParam Arguments[1] = {FWebJSParam(Error)};
-	InvokeJSFunction(FunctionId, 1, Arguments, true);
+	CefRefPtr<CefListValue> FunctionArguments = CefListValue::Create();
+	FunctionArguments->SetString(0, *Error);
+	InvokeJSFunction(FunctionId, FunctionArguments, true);
 }
 
 void FWebJSScripting::InvokeJSFunction(FGuid FunctionId, int32 ArgCount, FWebJSParam Arguments[], bool bIsError)
@@ -423,7 +332,11 @@ void FWebJSScripting::InvokeJSFunction(FGuid FunctionId, int32 ArgCount, FWebJSP
 	{
 		SetConverted(FunctionArguments, i, Arguments[i]);
 	}
+	InvokeJSFunction(FunctionId, FunctionArguments, bIsError);
+}
 
+void FWebJSScripting::InvokeJSFunction(FGuid FunctionId, const CefRefPtr<CefListValue>& FunctionArguments, bool bIsError)
+{
 	CefRefPtr<CefProcessMessage> Message = CefProcessMessage::Create(TEXT("UE::ExecuteJSFunction"));
 	CefRefPtr<CefListValue> MessageArguments = Message->GetArgumentList();
 	MessageArguments->SetString(0, *FunctionId.ToString(EGuidFormats::Digits));
diff --git a/Engine/Source/Runtime/WebBrowser/Private/WebJSScripting.h b/Engine/Source/Runtime/WebBrowser/Private/WebJSScripting.h
--- a/Engine/Source/Runtime/WebBrowser/Private/WebJSScripting.h
+++ b/Engine/Source/Runtime/WebBrowser/Private/WebJSScripting.h
@@ -93,12 +93,12 @@ public:
 		}
 	}
 
-	FWebJSParam ConvertResult(UProperty* ResultProperty, uint8* Data);
 	void BindCefBrowser(CefRefPtr<CefBrowser> Browser);
 	void UnbindCefBrowser();
 	CefRefPtr<CefDictionaryValue> GetPermanentBindings();
 
 	void InvokeJSFunction(FGuid FunctionId, int32 ArgCount, FWebJSParam Arguments[], bool bIsError=false);
+	void InvokeJSFunction(FGuid FunctionId, const CefRefPtr<CefListValue>& FunctionArguments, bool bIsError=false);
 	void InvokeJSErrorResult(FGuid FunctionId, const FString& Error);
 
 public:
