[INTEGRATE] Fix support for overlay-based popup menus. This is required for fullscreen applications where opening a popup window for the menu is not at option.
OPP-3860

Merging CL#2627615  using UE4-To-UE4-LauncherDev.

diff --git a/Engine/Source/Runtime/WebBrowser/Private/SWebBrowser.cpp b/Engine/Source/Runtime/WebBrowser/Private/SWebBrowser.cpp
--- a/Engine/Source/Runtime/WebBrowser/Private/SWebBrowser.cpp
+++ b/Engine/Source/Runtime/WebBrowser/Private/SWebBrowser.cpp
@@ -50,6 +50,7 @@ void SWebBrowser::Construct(const FArguments& InArgs, const TSharedPtr<IWebBrows
 	OnCreateWindow = InArgs._OnCreateWindow;
 	OnCloseWindow = InArgs._OnCloseWindow;
 	AddressBarUrl = FText::FromString(InArgs._InitialURL);
+	PopupMenuMethod = InArgs._PopupMenuMethod;
 
 	BrowserWindow = InWebBrowserWindow;
 	if(!BrowserWindow.IsValid())
@@ -75,9 +76,8 @@ void SWebBrowser::Construct(const FArguments& InArgs, const TSharedPtr<IWebBrows
 			InArgs._ShowErrorMessage,
 			InArgs._BackgroundColor
 			);		
+		}
 	}
-	}
-
 
 	ChildSlot
 	[
@@ -374,7 +374,6 @@ EVisibility SWebBrowser::GetLoadingThrobberVisibility() const
 	return EVisibility::Hidden;
 }
 
-
 void SWebBrowser::HandleBrowserWindowDocumentStateChanged(EWebBrowserDocumentState NewState)
 {
 	switch (NewState)
@@ -500,6 +499,8 @@ void SWebBrowser::UnbindUObject(const FString& Name, UObject* Object, bool bIsPe
 
 void SWebBrowser::HandleShowPopup(const FIntRect& PopupSize)
 {
+	check(!PopupMenuPtr.IsValid())
+
 	TSharedPtr<SViewport> MenuContent;
 	SAssignNew(MenuContent, SViewport)
 				.ViewportSize(PopupSize.Size())
@@ -517,19 +518,26 @@ void SWebBrowser::HandleShowPopup(const FIntRect& PopupSize)
 		const FGeometry& BrowserGeometry = WidgetPath.Widgets.Last().Geometry;
 		const FVector2D NewPosition = BrowserGeometry.LocalToAbsolute(PopupSize.Min);
 
-		// Open the pop-up
+
+		// Open the pop-up. The popup method will be queried from the widget path passed in.
 		TSharedPtr<IMenu> NewMenu = FSlateApplication::Get().PushMenu(ViewportWidget.ToSharedRef(), WidgetPath, MenuContentRef, NewPosition, FPopupTransitionEffect( FPopupTransitionEffect::ComboButton ), false);
+		NewMenu->GetOnMenuDismissed().AddSP(this, &SWebBrowser::HandleMenuDismissed);
 		PopupMenuPtr = NewMenu;
-		check(NewMenu.IsValid() && NewMenu->GetOwnedWindow().IsValid());
 	}
 
 }
 
+void SWebBrowser::HandleMenuDismissed(TSharedRef<IMenu>)
+{
+	PopupMenuPtr.Reset();
+}
+
 void SWebBrowser::HandleDismissPopup()
 {
 	if (PopupMenuPtr.IsValid())
 	{
 		PopupMenuPtr.Pin()->Dismiss();
+		FSlateApplication::Get().SetKeyboardFocus(ViewportWidget, EFocusCause::SetDirectly);
 	}
 }
 
diff --git a/Engine/Source/Runtime/WebBrowser/Private/WebBrowserHandler.cpp b/Engine/Source/Runtime/WebBrowser/Private/WebBrowserHandler.cpp
--- a/Engine/Source/Runtime/WebBrowser/Private/WebBrowserHandler.cpp
+++ b/Engine/Source/Runtime/WebBrowser/Private/WebBrowserHandler.cpp
@@ -12,7 +12,6 @@
 #if WITH_CEF3
 FWebBrowserHandler::FWebBrowserHandler()
 	: ShowErrorMessage(true)
-	, bShowPopupRequested(false)
 {
 }
 
@@ -259,29 +258,21 @@ void FWebBrowserHandler::OnCursorChange(CefRefPtr<CefBrowser> Browser, CefCursor
 void FWebBrowserHandler::OnPopupShow(CefRefPtr<CefBrowser> Browser, bool bShow)
 {
 	TSharedPtr<FWebBrowserWindow> BrowserWindow = BrowserWindowPtr.Pin();
-	bShowPopupRequested = bShow;
 
 	if (BrowserWindow.IsValid())
 	{
-		if (!bShow)
-		{
-			BrowserWindow->HidePopup();
-		}
+		BrowserWindow->ShowPopupMenu(bShow);
 	}
 
 }
 
 void FWebBrowserHandler::OnPopupSize(CefRefPtr<CefBrowser> Browser, const CefRect& Rect)
 {
-	if (Rect.width > 0 && Rect.height > 0)
+	TSharedPtr<FWebBrowserWindow> BrowserWindow = BrowserWindowPtr.Pin();
+
+	if (BrowserWindow.IsValid())
 	{
-		PopupRect = Rect;
-		TSharedPtr<FWebBrowserWindow> BrowserWindow = BrowserWindowPtr.Pin();
-		if (BrowserWindow.IsValid() && bShowPopupRequested)
-		{
-			BrowserWindow->ShowPopup(PopupRect);
-			bShowPopupRequested = false;
-		}
+		BrowserWindow->SetPopupMenuPosition(Rect);
 	}
 }
 
diff --git a/Engine/Source/Runtime/WebBrowser/Private/WebBrowserHandler.h b/Engine/Source/Runtime/WebBrowser/Private/WebBrowserHandler.h
--- a/Engine/Source/Runtime/WebBrowser/Private/WebBrowserHandler.h
+++ b/Engine/Source/Runtime/WebBrowser/Private/WebBrowserHandler.h
@@ -208,12 +208,6 @@ private:
 	/** Whether to show an error message in case of loading errors. */
 	bool ShowErrorMessage;
 
-	/** Size and position of the popup widget. */
-	CefRect PopupRect;
-
-	/** Set to true to show the popup widget on next popup size change. */
-	bool bShowPopupRequested;
-
 	// Include the default reference counting implementation.
 	IMPLEMENT_REFCOUNTING(FWebBrowserHandler);
 };
diff --git a/Engine/Source/Runtime/WebBrowser/Private/WebBrowserWindow.cpp b/Engine/Source/Runtime/WebBrowser/Private/WebBrowserWindow.cpp
--- a/Engine/Source/Runtime/WebBrowser/Private/WebBrowserWindow.cpp
+++ b/Engine/Source/Runtime/WebBrowser/Private/WebBrowserWindow.cpp
@@ -715,19 +715,21 @@ void FWebBrowserWindow::SetToolTip(const CefString& CefToolTip)
 	}
 }
 
-
 bool FWebBrowserWindow::GetViewRect(CefRect& Rect)
 {
+	Rect.x = 0;
+	Rect.y = 0;
 	if (ViewportSize == FIntPoint::ZeroValue)
 	{
-		return false;
+		// We need to return a dummy size until we know the final size of the viewport. Otherwise, popup menus may not work.
+		Rect.width = 1024;
+		Rect.height = 1024;
+	}
+	else
+	{
+		Rect.width = ViewportSize.X;
+		Rect.height = ViewportSize.Y;
 	}
-	
-	Rect.x = 0;
-	Rect.y = 0;
-	Rect.width = ViewportSize.X;
-	Rect.height = ViewportSize.Y;
-
 	return true;
 }
 
@@ -772,6 +774,15 @@ void FWebBrowserWindow::OnPaint(CefRenderHandler::PaintElementType Type, const C
 		UpdatableTextures[Type]->UpdateTextureThreadSafeRaw(Width, Height, Buffer, Dirty);
 	}
 
+	if (Type == PET_POPUP && bShowPopupRequested)
+	{
+		bShowPopupRequested = false;
+		bPopupHasFocus = true;
+		FIntPoint PopupSize = FIntPoint(Width, Height);
+		FIntRect PopupRect = FIntRect(PopupPosition, PopupPosition+PopupSize);
+		OnShowPopup().Broadcast(PopupRect);
+	}
+
 	bIsInitialized = true;
 	NeedsRedrawEvent.Broadcast();
 }
@@ -965,13 +976,12 @@ CefMouseEvent FWebBrowserWindow::GetCefMouseEvent(const FGeometry& MyGeometry, c
 {
 	CefMouseEvent Event;
 	FVector2D LocalPos = MyGeometry.AbsoluteToLocal(MouseEvent.GetScreenSpacePosition());
-	Event.x = LocalPos.X;
-	Event.y = LocalPos.Y;
 	if (bIsPopup)
 	{
-		Event.x += PopupRect.Min.X;
-		Event.y += PopupRect.Min.Y;
+		LocalPos += PopupPosition;
 	}
+	Event.x = LocalPos.X;
+	Event.y = LocalPos.Y;
 	Event.modifiers = GetCefMouseModifiers(MouseEvent);
 	return Event;
 }
@@ -1107,17 +1117,24 @@ void FWebBrowserWindow::OnBrowserClosed()
 	Handler = nullptr;
 }
 
-void FWebBrowserWindow::ShowPopup(CefRect CefPopupSize)
+void FWebBrowserWindow::SetPopupMenuPosition(CefRect CefPopupSize)
 {
-	PopupRect = FIntRect(CefPopupSize.x, CefPopupSize.y, CefPopupSize.x+CefPopupSize.width, CefPopupSize.y+CefPopupSize.height);
-	bPopupHasFocus = true;
-	OnShowPopup().Broadcast(PopupRect);
+	// We only store the position, as the size will be provided ib the OnPaint call.
+	PopupPosition = FIntPoint(CefPopupSize.x, CefPopupSize.y);
 }
 
-void FWebBrowserWindow::HidePopup()
+void FWebBrowserWindow:: ShowPopupMenu(bool bShow)
 {
-	bPopupHasFocus = false;
-	OnDismissPopup().Broadcast();
+	if (bShow)
+	{
+		bShowPopupRequested = true; // We have to delay showing the popup until we get the first OnPaint on it.
+	}
+	else
+	{
+		bPopupHasFocus = false;
+		bShowPopupRequested = false;
+		OnDismissPopup().Broadcast();
+	}
 }
 
 
diff --git a/Engine/Source/Runtime/WebBrowser/Private/WebBrowserWindow.h b/Engine/Source/Runtime/WebBrowser/Private/WebBrowserWindow.h
--- a/Engine/Source/Runtime/WebBrowser/Private/WebBrowserWindow.h
+++ b/Engine/Source/Runtime/WebBrowser/Private/WebBrowserWindow.h
@@ -318,16 +318,19 @@ private:
 	void OnBrowserClosed();
 
 	/**
-	 * Called to show the popup widget
+	 * Called to set the popup menu location. Note that CEF also passes a size to this method, 
+	 * which is ignored as the correct size is usually not known until inside OnPaint.
 	 *
- 	 * @param PopupSize The size and location of the popup widget.
+ 	 * @param PopupSize The location of the popup widget.
 	 */
-	void ShowPopup(CefRect PopupSize);
+	void SetPopupMenuPosition(CefRect PopupSize);
 
 	/**
-	 * Called to hide the popup widget again
+	 * Called to request that the popup widget is shown or hidden.
+	 *
+ 	 * @param bShow true for showing the popup, false for hiding it.
 	 */
-	void HidePopup();
+	void ShowPopupMenu(bool bShow);
 
 public:
 
@@ -458,7 +461,7 @@ private:
 	/** Delegate for handaling requests to show the popup menu. */
 	FOnShowPopup ShowPopupEvent;
 
-	/** Delegate for handaling requests to close new windows that were created. */
+	/** Delegate for hanaling requests to dismiss the current popup menu. */
 	FOnDismissPopup DismissPopupEvent;
 
 	/** Tracks the current mouse cursor */
@@ -482,7 +485,8 @@ private:
 	bool bMainHasFocus;
 	bool bPopupHasFocus;
 
-	FIntRect PopupRect;
+	FIntPoint PopupPosition;
+	bool bShowPopupRequested;
 
 	/** Handling of passing and marshalling messages for JS integration is delegated to a helper class*/
 	TSharedPtr<FWebJSScripting> Scripting;
diff --git a/Engine/Source/Runtime/WebBrowser/Public/SWebBrowser.h b/Engine/Source/Runtime/WebBrowser/Public/SWebBrowser.h
--- a/Engine/Source/Runtime/WebBrowser/Public/SWebBrowser.h
+++ b/Engine/Source/Runtime/WebBrowser/Public/SWebBrowser.h
@@ -30,6 +30,7 @@ public:
 		, _SupportsTransparency(false)
 		, _SupportsThumbMouseButtonNavigation(false)
 		, _BackgroundColor(255,255,255,255)
+		, _PopupMenuMethod(TOptional<EPopupMethod>())
 		, _ViewportSize(FVector2D::ZeroVector)
 	{ }
 
@@ -60,6 +61,9 @@ public:
 		/** Opaque background color used before a document is loaded and when no document color is specified. */
 		SLATE_ARGUMENT(FColor, BackgroundColor)
 
+		/** Override the popup menu method used for popup menus. If not set, parent widgets will be queried instead. */
+		SLATE_ARGUMENT(TOptional<EPopupMethod>, PopupMenuMethod)
+
 		/** Desired size of the web browser viewport. */
 		SLATE_ATTRIBUTE(FVector2D, ViewportSize);
 
@@ -255,6 +259,14 @@ private:
 	/** Callback for hiding the popup menu */
 	void HandleDismissPopup();
 
+	/** Callback from the popup menu notifiying it has been dismissed */
+	void HandleMenuDismissed(TSharedRef<IMenu>);
+
+	virtual TOptional<EPopupMethod> OnQueryPopupMethod() const override
+	{
+		return PopupMenuMethod;
+	}
+
 private:
 
 	/** Interface for dealing with a web browser window. */
@@ -274,6 +286,9 @@ private:
 	 */
 	TWeakPtr<IMenu> PopupMenuPtr;
 
+	/** Can be set to override the popup menu method used for popup menus. If not set, parent widgets will be queried instead. */
+	TOptional<EPopupMethod> PopupMenuMethod;
+
 	/** The url that appears in the address bar which can differ from the url of the loaded page */
 	FText AddressBarUrl;
 
