[INTEGRATE] Fix greedy templated constructor to only apply to non-pointer args. (Allows passing pointer to UObject-derived objects to FWebJSFunction without casting to UObject*)
Add support for passing TArray<T> as an argument to a FWebJSFunction.
Merging CL#2627633 using UE4-To-UE4-LauncherDev

diff --git a/Engine/Source/Runtime/WebBrowser/Private/WebJSFunction.cpp b/Engine/Source/Runtime/WebBrowser/Private/WebJSFunction.cpp
--- a/Engine/Source/Runtime/WebBrowser/Private/WebJSFunction.cpp
+++ b/Engine/Source/Runtime/WebBrowser/Private/WebJSFunction.cpp
@@ -19,10 +19,17 @@
 
 FWebJSParam::~FWebJSParam()
 {
-	// Since the FString member is in a union, it may or may not be valid, so we have to call its destructor manually.
-	if (Tag == PTYPE_STRING)
+	// Since the FString and TArray members are in a union, they may or may not be valid, so we have to call the destructors manually.
+	switch (Tag)
 	{
-		delete StringValue;
+		case PTYPE_STRING:
+			delete StringValue;
+			break;
+		case PTYPE_ARRAY:
+			delete ArrayValue;
+			break;
+		default:
+			break;
 	}
 }
 
@@ -52,6 +59,9 @@ FWebJSParam::FWebJSParam(const FWebJSParam& Other)
 			StructValue.TypeInfo = Other.StructValue.TypeInfo;
 			StructValue.StructPtr = Other.StructValue.StructPtr;
 			break;
+		case PTYPE_ARRAY:
+			ArrayValue = new TArray<FWebJSParam>(*Other.ArrayValue);
+			break;
 	}
 }
 
diff --git a/Engine/Source/Runtime/WebBrowser/Private/WebJSScripting.h b/Engine/Source/Runtime/WebBrowser/Private/WebJSScripting.h
--- a/Engine/Source/Runtime/WebBrowser/Private/WebJSScripting.h
+++ b/Engine/Source/Runtime/WebBrowser/Private/WebJSScripting.h
@@ -88,6 +88,15 @@ public:
 				CefRefPtr<CefDictionaryValue> ConvertedStruct = ConvertStruct(Param.StructValue.TypeInfo, Param.StructValue.StructPtr);
 				return Container->SetDictionary(Key, ConvertedStruct);
 			}
+			case FWebJSParam::PTYPE_ARRAY:
+			{
+				CefRefPtr<CefListValue> ConvertedArray = CefListValue::Create();
+				for(int i=0; i < Param.ArrayValue->Num(); ++i)
+				{
+					SetConverted(ConvertedArray, i, (*Param.ArrayValue)[i]);
+				}
+				return Container->SetList(Key, ConvertedArray);
+			}
 			default:
 				return false;
 		}
diff --git a/Engine/Source/Runtime/WebBrowser/Public/WebJSFunction.h b/Engine/Source/Runtime/WebBrowser/Public/WebJSFunction.h
--- a/Engine/Source/Runtime/WebBrowser/Public/WebJSFunction.h
+++ b/Engine/Source/Runtime/WebBrowser/Public/WebJSFunction.h
@@ -2,13 +2,14 @@
 
 #pragma once
 #include "CoreUObject.h"
+#include "EnableIf.h"
 #include "WebJSFunction.generated.h"
 
 class FWebJSScripting;
 
 struct WEBBROWSER_API FWebJSParam
 {
-	enum { PTYPE_NULL, PTYPE_BOOL, PTYPE_INT, PTYPE_DOUBLE, PTYPE_STRING, PTYPE_OBJECT, PTYPE_STRUCT, } Tag;
+	enum { PTYPE_NULL, PTYPE_BOOL, PTYPE_INT, PTYPE_DOUBLE, PTYPE_STRING, PTYPE_OBJECT, PTYPE_STRUCT, PTYPE_ARRAY } Tag;
 	union
 	{
 		bool BoolValue;
@@ -17,9 +18,10 @@ struct WEBBROWSER_API FWebJSParam
 		UObject* ObjectValue;
 		const FString* StringValue;
 		struct {
-			UStruct* TypeInfo;
 			const void* StructPtr;
+			UStruct* TypeInfo;
 		} StructValue;
+		TArray<FWebJSParam>* ArrayValue;
 	};
 
 	FWebJSParam() : Tag(PTYPE_NULL) {}
@@ -37,8 +39,21 @@ struct WEBBROWSER_API FWebJSParam
 	FWebJSParam(const FString& Value) : Tag(PTYPE_STRING), StringValue(new FString(Value)) {}
 	FWebJSParam(const TCHAR* Value) : Tag(PTYPE_STRING), StringValue(new FString(Value)) {}
 	FWebJSParam(UObject* Value) : Tag(PTYPE_OBJECT), ObjectValue(Value) {}
-	FWebJSParam(UStruct* Struct, const void* Value) : Tag(PTYPE_STRUCT), StructValue({Struct, Value}) {}
-	template <typename T> FWebJSParam(const T& Value) : Tag(PTYPE_STRUCT), StructValue({T::StaticStruct(), &Value}) {}
+	FWebJSParam(const void* Value, UStruct* Struct) : Tag(PTYPE_STRUCT), StructValue({Value, Struct}) {}
+	template <typename T> FWebJSParam(const T& Value,
+		typename TEnableIf<!TIsPointerType<T>::Value, UStruct>::Type* InTypeInfo=T::StaticStruct())
+		: Tag(PTYPE_STRUCT)
+		, StructValue({&Value, InTypeInfo}) {}
+	template <typename T> FWebJSParam(const TArray<T>& Value)
+		: Tag(PTYPE_ARRAY)
+	{
+		ArrayValue = new TArray<FWebJSParam>();
+		ArrayValue->Reserve(Value.Num());
+		for(T Item : Value)
+		{
+			ArrayValue->Add(FWebJSParam(Item));
+		}
+	}
 	FWebJSParam(const FWebJSParam& Other);
 	~FWebJSParam();
 
