[INTEGRATE] SWebBrowser refactor and cleanups:
* Untangle mutual dependencies between objects and constructor order (OPP-3858):
	* FWebBrowserWindow no longer keeps a reference to FWebBrowserHandler.
	* CefBrowserWindow is created before creating FWebBrowserWindow and passed to its constructor.
	* FWebBrowserViewport no longer has a reference back to the viewport widget.
* Cleaned how some platform specific paths were defined during module initialization.
* Sets current thread name (back) to GameThread after initializing CEF (UE-5165)
* Adopted changes from UT (with some modifications):
	* Enable on-disk caching and persisting cookies when requested by server.
	* Added product name and engine version to user agent request header.
* Ensure UStruct parameters passed by value to FWebJSFunctions are copied when constructing the argument list and not just their address.
*SWebBrowser scripting: Add support for passing TMap<FString, T> to FWebJSFunction objects.
	Note: since the struct serializer does not yet support maps, this only works for arguments passed directly to the function object and not UStruct members nor function return values. OPP-3861
* Add support for passing a reference to the JS promise object implicitly created when calling UFunctions to provide for a way to report runtime errors back to it instead of just passing the return value.
	* This also allows for implementing async functionality as the promise object can be saved and used after the UFunction has returned.

Merging CL#2633955, CL#2634077, CL#2634226, CL#2635540, CL#2636167, and CL#2636262 using UE4-To-UE4-LauncherDev

diff --git a/Engine/Source/Runtime/WebBrowser/Private/SWebBrowser.cpp b/Engine/Source/Runtime/WebBrowser/Private/SWebBrowser.cpp
--- a/Engine/Source/Runtime/WebBrowser/Private/SWebBrowser.cpp
+++ b/Engine/Source/Runtime/WebBrowser/Private/SWebBrowser.cpp
@@ -68,8 +68,6 @@ void SWebBrowser::Construct(const FArguments& InArgs, const TSharedPtr<IWebBrows
 		BrowserWindow = IWebBrowserModule::Get().GetSingleton()->CreateBrowserWindow(
 			OSWindowHandle,
 			InArgs._InitialURL,
-			InArgs._ViewportSize.Get().X,
-			InArgs._ViewportSize.Get().Y,
 			InArgs._SupportsTransparency,
 			InArgs._SupportsThumbMouseButtonNavigation,
 			InArgs._ContentsToLoad,
@@ -189,7 +187,7 @@ void SWebBrowser::Construct(const FArguments& InArgs, const TSharedPtr<IWebBrows
 		BrowserWindow->OnBeforePopup().BindSP(this, &SWebBrowser::HandleBeforePopup);
 		BrowserWindow->OnShowPopup().AddSP(this, &SWebBrowser::HandleShowPopup);
 		BrowserWindow->OnDismissPopup().AddSP(this, &SWebBrowser::HandleDismissPopup);
-		BrowserViewport = MakeShareable(new FWebBrowserViewport(BrowserWindow, ViewportWidget));
+		BrowserViewport = MakeShareable(new FWebBrowserViewport(BrowserWindow));
 		ViewportWidget->SetViewportInterface(BrowserViewport.ToSharedRef());
 	}
 	else
@@ -508,7 +506,7 @@ void SWebBrowser::HandleShowPopup(const FIntRect& PopupSize)
 				.EnableBlending(false)
 				.IgnoreTextureAlpha(true)
 				.Visibility(EVisibility::Visible);
-	MenuViewport = MakeShareable(new FWebBrowserViewport(BrowserWindow, MenuContent, true));
+	MenuViewport = MakeShareable(new FWebBrowserViewport(BrowserWindow, true));
 	MenuContent->SetViewportInterface(MenuViewport.ToSharedRef());
 	FWidgetPath WidgetPath;
 	FSlateApplication::Get().GeneratePathToWidgetUnchecked(ViewportWidget.ToSharedRef(), WidgetPath);
diff --git a/Engine/Source/Runtime/WebBrowser/Private/WebBrowserHandler.cpp b/Engine/Source/Runtime/WebBrowser/Private/WebBrowserHandler.cpp
--- a/Engine/Source/Runtime/WebBrowser/Private/WebBrowserHandler.cpp
+++ b/Engine/Source/Runtime/WebBrowser/Private/WebBrowserHandler.cpp
@@ -52,50 +52,40 @@ bool FWebBrowserHandler::OnTooltip(CefRefPtr<CefBrowser> Browser, CefString& Tex
 
 void FWebBrowserHandler::OnAfterCreated(CefRefPtr<CefBrowser> Browser)
 {
-	TSharedPtr<FWebBrowserWindow> BrowserWindow = BrowserWindowPtr.Pin();
-
-	if (BrowserWindow.IsValid())
+	if(Browser->IsPopup())
 	{
-
-		BrowserWindow->BindCefBrowser(Browser);
-	}
-	else
-	{
-		if(Browser->IsPopup())
+		TSharedPtr<FWebBrowserWindow> BrowserWindowParent = BrowserWindowParentPtr.Pin();
+		if(BrowserWindowParent.IsValid() && BrowserWindowParent->SupportsNewWindows())
 		{
-			TSharedPtr<FWebBrowserWindow> BrowserWindowParent = BrowserWindowParentPtr.Pin();
-			if(BrowserWindowParent.IsValid() && BrowserWindowParent->SupportsNewWindows())
+			TSharedPtr<FWebBrowserWindowInfo> NewBrowserWindowInfo = MakeShareable(new FWebBrowserWindowInfo(Browser, this));
+			TSharedPtr<IWebBrowserWindow> NewBrowserWindow = IWebBrowserModule::Get().GetSingleton()->CreateBrowserWindow( 
+				BrowserWindowParent,
+				NewBrowserWindowInfo
+				);	
+
 			{
-				TSharedPtr<FWebBrowserWindowInfo> NewBrowserWindowInfo = MakeShareable(new FWebBrowserWindowInfo(Browser, this));
-				TSharedPtr<IWebBrowserWindow> NewBrowserWindow = IWebBrowserModule::Get().GetSingleton()->CreateBrowserWindow( 
-					BrowserWindowParent,
-					NewBrowserWindowInfo
-					);	
-
-				{
-					// @todo: At the moment we need to downcast since the handler does not support using the interface.
-					TSharedPtr<FWebBrowserWindow> HandlerSpecificBrowserWindow = StaticCastSharedPtr<FWebBrowserWindow>(NewBrowserWindow);
-					BrowserWindowPtr = HandlerSpecificBrowserWindow;
-				}
-
-				// Request a UI window for the browser.  If it is not created we do some cleanup.
-				bool bUIWindowCreated = BrowserWindowParent->RequestCreateWindow(NewBrowserWindow.ToSharedRef(), BrowserPopupFeatures);
-				if(!bUIWindowCreated)
-				{
-					NewBrowserWindow->CloseBrowser(true);
-				}
+				// @todo: At the moment we need to downcast since the handler does not support using the interface.
+				TSharedPtr<FWebBrowserWindow> HandlerSpecificBrowserWindow = StaticCastSharedPtr<FWebBrowserWindow>(NewBrowserWindow);
+				BrowserWindowPtr = HandlerSpecificBrowserWindow;
 			}
-			else
+
+			// Request a UI window for the browser.  If it is not created we do some cleanup.
+			bool bUIWindowCreated = BrowserWindowParent->RequestCreateWindow(NewBrowserWindow.ToSharedRef(), BrowserPopupFeatures);
+			if(!bUIWindowCreated)
 			{
-				Browser->GetHost()->CloseBrowser(true);
+				NewBrowserWindow->CloseBrowser(true);
 			}
 		}
+		else
+		{
+			Browser->GetHost()->CloseBrowser(true);
+		}
 	}
 }
 
  bool FWebBrowserHandler::DoClose(CefRefPtr<CefBrowser> Browser)
  {
-	 TSharedPtr<FWebBrowserWindow> BrowserWindow = BrowserWindowPtr.Pin();
+	TSharedPtr<FWebBrowserWindow> BrowserWindow = BrowserWindowPtr.Pin();
 	if(BrowserWindow.IsValid())
 	{
 		BrowserWindow->OnBrowserClosing();
@@ -226,8 +216,13 @@ bool FWebBrowserHandler::GetViewRect(CefRefPtr<CefBrowser> Browser, CefRect& Rec
 	{
 		return BrowserWindow->GetViewRect(Rect);
 	}
-
-	return false;
+	else
+	{
+		// Seems that if we return false from here or without a size > 0, popup menus will not show up.
+		Rect.width = 800;
+		Rect.height = 600;
+		return true;
+	}
 }
 
 void FWebBrowserHandler::OnPaint(CefRefPtr<CefBrowser> Browser,
diff --git a/Engine/Source/Runtime/WebBrowser/Private/WebBrowserSingleton.cpp b/Engine/Source/Runtime/WebBrowser/Private/WebBrowserSingleton.cpp
--- a/Engine/Source/Runtime/WebBrowser/Private/WebBrowserSingleton.cpp
+++ b/Engine/Source/Runtime/WebBrowser/Private/WebBrowserSingleton.cpp
@@ -9,30 +9,120 @@
 #include "IPlatformTextField.h"
 #include "IVirtualKeyboardEntry.h"
 #include "SlateApplication.h"
+#include "Runtime/Launch/Resources/Version.h"
 
 #if WITH_CEF3
-#if PLATFORM_WINDOWS
-	#include "AllowWindowsPlatformTypes.h"
+#	if PLATFORM_WINDOWS
+#		include "AllowWindowsPlatformTypes.h"
+#	endif
+#	pragma push_macro("OVERRIDE")
+#		undef OVERRIDE // cef headers provide their own OVERRIDE macro
+#		include "include/cef_app.h"
+#	pragma pop_macro("OVERRIDE")
+#	if PLATFORM_WINDOWS
+#		include "HideWindowsPlatformTypes.h"
+#	endif
 #endif
-	#pragma push_macro("OVERRIDE")
-	#undef OVERRIDE // cef headers provide their own OVERRIDE macro
-	#include "include/cef_app.h"
-	#pragma pop_macro("OVERRIDE")
-#if PLATFORM_WINDOWS
-	#include "HideWindowsPlatformTypes.h"
+
+#if PLATFORM_MAC || PLATFORM_LINUX
+#	include <pthread.h>
+#endif
+
+// Define some platform-dependent file locations
+#define CEF3_BIN_DIR TEXT("Binaries/ThirdParty/CEF3")
+#if PLATFORM_WINDOWS && PLATFORM_64BITS
+#	define CEF3_RESOURCES_DIR CEF3_BIN_DIR TEXT("/Win64/Resources")
+#	define CEF3_SUBPROCES_EXE TEXT("Binaries/Win64/UnrealCEFSubProcess.exe")
+#elif PLATFORM_WINDOWS && PLATFORM_32BITS
+#	define CEF3_RESOURCES_DIR CEF3_BIN_DIR TEXT("/Win32/Resources")
+#	define CEF3_SUBPROCES_EXE TEXT("Binaries/Win32/UnrealCEFSubProcess.exe")
+#elif PLATFORM_MAC
+#	define CEF3_FRAMEWORK_DIR CEF3_BIN_DIR TEXT("/Mac/Chromium Embedded Framework.framework")
+#	define CEF3_RESOURCES_DIR CEF3_FRAMEWORK_DIR TEXT("/Resources")
+#	define CEF3_SUBPROCES_EXE TEXT("Binaries/Mac/UnrealCEFSubProcess.app/Contents/MacOS/UnrealCEFSubProcess")
+#elif PLATFORM_LINUX // @todo Linux
+#	define CEF3_RESOURCES_DIR CEF3_BIN_DIR TEXT("/Linux/Resources")
+#	define CEF3_SUBPROCES_EXE TEXT("Binaries/Linux/UnrealCEFSubProcess")
+#endif
+
+
+namespace {
+
+	/**
+	 * Helper function to set the current thread name, visible by the debugger.
+	 * @param ThreadName	Name to set
+	 */
+	void SetCurrentThreadName(char* ThreadName)
+	{
+#if PLATFORM_MAC
+		pthread_setname_np(ThreadName);
+#elif PLATFORM_LINUX
+		pthread_setname_np(pthread_self(), ThreadName);
+#elif PLATFORM_WINDOWS && !PLATFORM_SEH_EXCEPTIONS_DISABLED
+		/**
+		 * Code setting the thread name for use in the debugger.
+		 * Copied implementation from WindowsRunnableThread as it is private.
+		 *
+		 * http://msdn.microsoft.com/en-us/library/xcb2z8hs.aspx
+		 */
+		const uint32 MS_VC_EXCEPTION=0x406D1388;
+
+		struct THREADNAME_INFO
+		{
+			uint32 dwType;		// Must be 0x1000.
+			LPCSTR szName;		// Pointer to name (in user addr space).
+			uint32 dwThreadID;	// Thread ID (-1=caller thread).
+			uint32 dwFlags;		// Reserved for future use, must be zero.
+		};
+
+		THREADNAME_INFO ThreadNameInfo = {0x1000, ThreadName, -1, 0};
+
+		__try
+		{
+			RaiseException( MS_VC_EXCEPTION, 0, sizeof(ThreadNameInfo)/sizeof(ULONG_PTR), (ULONG_PTR*)&ThreadNameInfo );
+		}
+		__except( EXCEPTION_EXECUTE_HANDLER )
+		CA_SUPPRESS(6322)
+		{	
+		}
 #endif
+	}
+
+#if PLATFORM_MAC
+	// OSX wants caches in a separate location from other app data
+	const TCHAR* ApplicationCacheDir()
+	{
+		static TCHAR Result[MAX_PATH] = TEXT("");
+		if (!Result[0])
+		{
+			SCOPED_AUTORELEASE_POOL;
+			NSString *CacheBaseDir = [NSSearchPathForDirectoriesInDomains(NSCachesDirectory, NSUserDomainMask, YES) objectAtIndex: 0];
+			NSString* BundleID = [[NSBundle mainBundle] bundleIdentifier];
+			if(!BundleID)
+			{
+				BundleID = [[NSProcessInfo processInfo] processName];
+			}
+			check(BundleID);
+			
+			NSString* AppCacheDir = [CacheBaseDir stringByAppendingPathComponent: BundleID];
+			FPlatformString::CFStringToTCHAR((CFStringRef)AppCacheDir, Result);
+		}
+		return Result;
+	}
 #endif
+}
 
 FWebBrowserSingleton::FWebBrowserSingleton()
 {
 #if WITH_CEF3
+	// The FWebBrowserSingleton must be initialized on the game thread
+	check(IsInGameThread());
+
 	// Provide CEF with command-line arguments.
 #if PLATFORM_WINDOWS
 	CefMainArgs MainArgs(hInstance);
-#elif PLATFORM_MAC || PLATFORM_LINUX
-	//TArray<FString> Args;
-	//int ArgCount = GSavedCommandLine.ParseIntoArray(Args, TEXT(" "), true);
-	CefMainArgs MainArgs(0, nullptr);
+#else
+	CefMainArgs MainArgs;
 #endif
 
 	// WebBrowserApp implements application-level callbacks.
@@ -43,27 +133,38 @@ FWebBrowserSingleton::FWebBrowserSingleton()
 	CefSettings Settings;
 	Settings.no_sandbox = true;
 	Settings.command_line_args_disabled = true;
+	Settings.log_severity = LOGSEVERITY_WARNING;
 
 	// Specify locale from our settings
 	FString LocaleCode = GetCurrentLocaleCode();
 	CefString(&Settings.locale) = *LocaleCode;
 
-	// Specify path to resources
-#if PLATFORM_WINDOWS
-	#if PLATFORM_64BITS
-	FString ResourcesPath(FPaths::Combine(*FPaths::EngineDir(), TEXT("Binaries/ThirdParty/CEF3/Win64/Resources")));
-	FString LocalesPath(FPaths::Combine(*FPaths::EngineDir(), TEXT("Binaries/ThirdParty/CEF3/Win64/Resources/locales")));
-	#else
-	FString ResourcesPath(FPaths::Combine(*FPaths::EngineDir(), TEXT("Binaries/ThirdParty/CEF3/Win32/Resources")));
-	FString LocalesPath(FPaths::Combine(*FPaths::EngineDir(), TEXT("Binaries/ThirdParty/CEF3/Win32/Resources/locales")));
-	#endif
-#elif PLATFORM_MAC
-	FString ResourcesPath(FPaths::Combine(*FPaths::EngineDir(), TEXT("Binaries/ThirdParty/CEF3/Mac/Chromium Embedded Framework.framework/Resources")));
-#elif PLATFORM_LINUX // @todo Linux
-	FString ResourcesPath(FPaths::Combine(*FPaths::EngineDir(), TEXT("Binaries/ThirdParty/CEF3/Linux/Resources")));
-	FString LocalesPath(FPaths::Combine(*FPaths::EngineDir(), TEXT("Binaries/ThirdParty/CEF3/Linux/Resources/locales")));
+	// Append engine version to the user agent string.
+	FString ProductVersion = FString::Printf( TEXT("%s UnrealEngineChrome/%s"), FApp::GetGameName(), ENGINE_VERSION_STRING);
+	CefString(&Settings.product_version) = *ProductVersion;
+
+	// Enable on disk cache
+#if PLATFORM_MAC
+	// OSX wants cache files in a special location
+	FString CachePath(FPaths::Combine(ApplicationCacheDir(), TEXT("webcache")));
+#else
+	FString CachePath(FPaths::Combine(*FPaths::GameSavedDir(), TEXT("webcache")));
 #endif
-#if !PLATFORM_MAC // On Mac Chromium ignores custom locales dir. Files need to be stored in Resources folder in the app bundle
+	CachePath = FPaths::ConvertRelativePathToFull(CachePath);
+	CefString(&Settings.cache_path) = *CachePath;
+
+	// Specify path to resources
+	FString ResourcesPath(FPaths::Combine(*FPaths::EngineDir(), CEF3_RESOURCES_DIR));
+	ResourcesPath = FPaths::ConvertRelativePathToFull(ResourcesPath);
+	if (!FPaths::DirectoryExists(ResourcesPath))
+	{
+		UE_LOG(LogWebBrowser, Error, TEXT("Chromium Resources information not found at: %s."), *ResourcesPath);
+	}
+	CefString(&Settings.resources_dir_path) = *ResourcesPath;
+
+#if !PLATFORM_MAC
+	// On Mac Chromium ignores custom locales dir. Files need to be stored in Resources folder in the app bundle
+	FString LocalesPath(FPaths::Combine(*ResourcesPath, TEXT("locales")));
 	LocalesPath = FPaths::ConvertRelativePathToFull(LocalesPath);
 	if (!FPaths::DirectoryExists(LocalesPath))
 	{
@@ -71,25 +172,9 @@ FWebBrowserSingleton::FWebBrowserSingleton()
 	}
 	CefString(&Settings.locales_dir_path) = *LocalesPath;
 #endif
-	ResourcesPath = FPaths::ConvertRelativePathToFull(ResourcesPath);
-	if (!FPaths::DirectoryExists(ResourcesPath))
-	{
-		UE_LOG(LogWebBrowser, Error, TEXT("Chromium Resources information not found at: %s."), *ResourcesPath);
-	}
-	CefString(&Settings.resources_dir_path) = *ResourcesPath;
 
 	// Specify path to sub process exe
-#if PLATFORM_WINDOWS
-	#if PLATFORM_64BITS
-	FString SubProcessPath(FPaths::Combine(*FPaths::EngineDir(), TEXT("Binaries/Win64/UnrealCEFSubProcess.exe")));
-	#else
-	FString SubProcessPath(FPaths::Combine(*FPaths::EngineDir(), TEXT("Binaries/Win32/UnrealCEFSubProcess.exe")));
-	#endif
-#elif PLATFORM_MAC
-	FString SubProcessPath(FPaths::Combine(*FPaths::EngineDir(), TEXT("Binaries/Mac/UnrealCEFSubProcess.app/Contents/MacOS/UnrealCEFSubProcess")));
-#else // @todo Linux
-	FString SubProcessPath(TEXT("UnrealCEFSubProcess"));
-#endif
+	FString SubProcessPath(FPaths::Combine(*FPaths::EngineDir(), CEF3_SUBPROCES_EXE));
 	SubProcessPath = FPaths::ConvertRelativePathToFull(SubProcessPath);
 	if (!FPaths::FileExists(SubProcessPath))
 	{
@@ -100,6 +185,9 @@ FWebBrowserSingleton::FWebBrowserSingleton()
 	// Initialize CEF.
 	bool bSuccess = CefInitialize(MainArgs, Settings, WebBrowserApp.get(), nullptr);
 	check(bSuccess);
+
+	// Set the thread name back to GameThread.
+	SetCurrentThreadName(TCHAR_TO_ANSI( *(FName( NAME_GameThread ).GetPlainNameString()) ));
 #endif
 }
 
@@ -132,12 +220,14 @@ FWebBrowserSingleton::~FWebBrowserSingleton()
 			WindowInterfaces[Index].Pin()->CloseBrowser(true);
 		}
 	}
+
 	// Just in case, although we deallocate WebBrowserApp right after this.
 	WebBrowserApp->OnRenderProcessThreadCreated().Unbind();
 	// CefRefPtr takes care of delete
 	WebBrowserApp = nullptr;
 	// Shut down CEF.
 	CefShutdown();
+
 #endif
 }
 
@@ -148,31 +238,24 @@ TSharedPtr<IWebBrowserWindow> FWebBrowserSingleton::CreateBrowserWindow(
 {
 #if WITH_CEF3
 
-	//  @todo: Width/Height should be obtained when requesting a UI window from user code which happens later.
-	int32 Width = 800;
-	int32 Height = 600;
 	TOptional<FString> ContentsToLoad;
 
 	bool bShowErrorMessage = BrowserWindowParent->IsShowingErrorMessages();
 	bool bThumbMouseButtonNavigation = BrowserWindowParent->IsThumbMouseButtonNavigationEnabled();
 	bool bUseTransparency = BrowserWindowParent->UseTransparency();
 	FString InitialURL = BrowserWindowInfo->Browser->GetMainFrame()->GetURL().ToWString().c_str();
-	TSharedPtr<FWebBrowserWindow> NewBrowserWindow(new FWebBrowserWindow(FIntPoint(Width, Height), InitialURL, ContentsToLoad, bShowErrorMessage, bThumbMouseButtonNavigation, bUseTransparency));
-	
-	NewBrowserWindow->SetHandler(BrowserWindowInfo->Handler);
-	NewBrowserWindow->BindCefBrowser(BrowserWindowInfo->Browser);
+	TSharedPtr<FWebBrowserWindow> NewBrowserWindow(new FWebBrowserWindow(BrowserWindowInfo->Browser, InitialURL, ContentsToLoad, bShowErrorMessage, bThumbMouseButtonNavigation, bUseTransparency));
+	BrowserWindowInfo->Handler->SetBrowserWindow(NewBrowserWindow);
 
 	WindowInterfaces.Add(NewBrowserWindow);
 	return NewBrowserWindow;
 #endif
-	return NULL;
+	return nullptr;
 }
 
 TSharedPtr<IWebBrowserWindow> FWebBrowserSingleton::CreateBrowserWindow(
 	void* OSWindowHandle, 
 	FString InitialURL, 
-	uint32 Width, 
-	uint32 Height, 
 	bool bUseTransparency,
 	bool bThumbMouseButtonNavigation,
 	TOptional<FString> ContentsToLoad, 
@@ -180,12 +263,6 @@ TSharedPtr<IWebBrowserWindow> FWebBrowserSingleton::CreateBrowserWindow(
 	FColor BackgroundColor)
 {
 #if WITH_CEF3
-	// Create new window
-	TSharedPtr<FWebBrowserWindow> NewWindow(new FWebBrowserWindow(FIntPoint(Width, Height), InitialURL, ContentsToLoad, ShowErrorMessage, bThumbMouseButtonNavigation, bUseTransparency));
-
-	// WebBrowserHandler implements browser-level callbacks.
-	CefRefPtr<FWebBrowserHandler> NewHandler(new FWebBrowserHandler);
-	NewWindow->SetHandler(NewHandler);
 
 	// Information used when creating the native window.
 	CefWindowHandle WindowHandle = (CefWindowHandle)OSWindowHandle; // TODO: check this is correct for all platforms
@@ -204,16 +281,23 @@ TSharedPtr<IWebBrowserWindow> FWebBrowserSingleton::CreateBrowserWindow(
 	// Disable plugins
 	BrowserSettings.plugins = STATE_DISABLED;
 
-	CefString URL = *InitialURL;
+
+	// WebBrowserHandler implements browser-level callbacks.
+	CefRefPtr<FWebBrowserHandler> NewHandler(new FWebBrowserHandler);
 
 	// Create the CEF browser window.
-	if (CefBrowserHost::CreateBrowser(WindowInfo, NewHandler.get(), URL, BrowserSettings, NULL))
+	CefRefPtr<CefBrowser> Browser = CefBrowserHost::CreateBrowserSync(WindowInfo, NewHandler.get(), *InitialURL, BrowserSettings, nullptr);
+	if (Browser.get())
 	{
-		WindowInterfaces.Add(NewWindow);
-		return NewWindow;
+		// Create new window
+		TSharedPtr<FWebBrowserWindow> NewBrowserWindow(new FWebBrowserWindow(Browser, InitialURL, ContentsToLoad, ShowErrorMessage, bThumbMouseButtonNavigation, bUseTransparency));
+		NewHandler->SetBrowserWindow(NewBrowserWindow);
+
+		WindowInterfaces.Add(NewBrowserWindow);
+		return NewBrowserWindow;
 	}
 #endif
-	return NULL;
+	return nullptr;
 }
 
 bool FWebBrowserSingleton::Tick(float DeltaTime)
@@ -253,3 +337,9 @@ FString FWebBrowserSingleton::GetCurrentLocaleCode()
 	}
 	return LocaleCode;
 }
+
+// Cleanup macros to avoid having them leak outside this source file
+#undef CEF3_BIN_DIR
+#undef CEF3_FRAMEWORK_DIR
+#undef CEF3_RESOURCES_DIR
+#undef CEF3_SUBPROCES_EXE
diff --git a/Engine/Source/Runtime/WebBrowser/Private/WebBrowserSingleton.h b/Engine/Source/Runtime/WebBrowser/Private/WebBrowserSingleton.h
--- a/Engine/Source/Runtime/WebBrowser/Private/WebBrowserSingleton.h
+++ b/Engine/Source/Runtime/WebBrowser/Private/WebBrowserSingleton.h
@@ -58,9 +58,7 @@ public:
 	TSharedPtr<IWebBrowserWindow> CreateBrowserWindow(
 		void* OSWindowHandle, 
 		FString InitialURL, 
-		uint32 Width, 
-		uint32 Height, 
-		bool bUseTransparency, 
+		bool bUseTransparency,
 		bool bThumbMouseButtonNavigation,
 		TOptional<FString> ContentsToLoad = TOptional<FString>(), 
 		bool ShowErrorMessage = true,
diff --git a/Engine/Source/Runtime/WebBrowser/Private/WebBrowserViewport.cpp b/Engine/Source/Runtime/WebBrowser/Private/WebBrowserViewport.cpp
--- a/Engine/Source/Runtime/WebBrowser/Private/WebBrowserViewport.cpp
+++ b/Engine/Source/Runtime/WebBrowser/Private/WebBrowserViewport.cpp
@@ -46,9 +46,14 @@ FReply FWebBrowserViewport::OnMouseButtonDown(const FGeometry& MyGeometry, const
 {
 	// Capture mouse on left button down so that you can drag out of the viewport
 	FReply Reply = WebBrowserWindow->OnMouseButtonDown(MyGeometry, MouseEvent, bIsPopup);
-	if (MouseEvent.GetEffectingButton() == EKeys::LeftMouseButton && ViewportWidget.IsValid())
+	if (MouseEvent.GetEffectingButton() == EKeys::LeftMouseButton)
 	{
-		return Reply.CaptureMouse(ViewportWidget.Pin().ToSharedRef());
+		const FWidgetPath* Path = MouseEvent.GetEventPath();
+		if (Path->IsValid())
+		{
+			TSharedRef<SWidget> TopWidget = Path->Widgets.Last().Widget;
+			return Reply.CaptureMouse(TopWidget);
+		}
 	}
 	return Reply;
 }
diff --git a/Engine/Source/Runtime/WebBrowser/Private/WebBrowserWindow.cpp b/Engine/Source/Runtime/WebBrowser/Private/WebBrowserWindow.cpp
--- a/Engine/Source/Runtime/WebBrowser/Private/WebBrowserWindow.cpp
+++ b/Engine/Source/Runtime/WebBrowser/Private/WebBrowserWindow.cpp
@@ -24,10 +24,11 @@ typedef FMacCursor FPlatformCursor;
 #else
 #endif
 
-FWebBrowserWindow::FWebBrowserWindow(FIntPoint InViewportSize, FString InUrl, TOptional<FString> InContentsToLoad, bool InShowErrorMessage, bool InThumbMouseButtonNavigation, bool InUseTransparency)
+FWebBrowserWindow::FWebBrowserWindow(CefRefPtr<CefBrowser> InBrowser, FString InUrl, TOptional<FString> InContentsToLoad, bool InShowErrorMessage, bool InThumbMouseButtonNavigation, bool InUseTransparency)
 	: DocumentState(EWebBrowserDocumentState::NoDocument)
+	, InternalCefBrowser(InBrowser)
 	, CurrentUrl(InUrl)
-	, ViewportSize(InViewportSize)
+	, ViewportSize(FIntPoint::ZeroValue)
 	, bIsClosing(false)
 	, bIsInitialized(false)
 	, ContentsToLoad(InContentsToLoad)
@@ -45,8 +46,10 @@ FWebBrowserWindow::FWebBrowserWindow(FIntPoint InViewportSize, FString InUrl, TO
 	, bIgnoreCharacterEvent(false)
 	, bMainHasFocus(false)
 	, bPopupHasFocus(false)
-	, Scripting(new FWebJSScripting)
+	, Scripting(new FWebJSScripting(InBrowser))
 {
+	check(InBrowser.get() != nullptr);
+
 	UpdatableTextures[0] = nullptr;
 	UpdatableTextures[1] = nullptr;
 }
@@ -663,15 +666,6 @@ void FWebBrowserWindow::ExecuteJavascript(const FString& Script)
 	}
 }
 
-void FWebBrowserWindow::SetHandler(CefRefPtr<FWebBrowserHandler> InHandler)
-{
-	if (InHandler.get())
-	{
-		Handler = InHandler;
-		Handler->SetBrowserWindow(SharedThis(this));
-		Handler->SetShowErrorMessage(ShowErrorMessage);
-	}
-}
 
 void FWebBrowserWindow::CloseBrowser(bool bForce)
 {
@@ -681,13 +675,6 @@ void FWebBrowserWindow::CloseBrowser(bool bForce)
 	}
 }
 
-void FWebBrowserWindow::BindCefBrowser(CefRefPtr<CefBrowser> Browser)
-{
-	check(Browser.get() == nullptr || InternalCefBrowser.get() == nullptr || InternalCefBrowser->IsSame(Browser));
-	Scripting->BindCefBrowser(Browser); // The scripting interface needs the browser too
-	InternalCefBrowser = Browser;
-}
-
 CefRefPtr<CefBrowser> FWebBrowserWindow::GetCefBrowser()
 {
 	return InternalCefBrowser;
@@ -1114,7 +1101,6 @@ void FWebBrowserWindow::OnBrowserClosed()
 
 	Scripting->UnbindCefBrowser();
 	InternalCefBrowser = nullptr;
-	Handler = nullptr;
 }
 
 void FWebBrowserWindow::SetPopupMenuPosition(CefRect CefPopupSize)
diff --git a/Engine/Source/Runtime/WebBrowser/Private/WebBrowserWindow.h b/Engine/Source/Runtime/WebBrowser/Private/WebBrowserWindow.h
--- a/Engine/Source/Runtime/WebBrowser/Private/WebBrowserWindow.h
+++ b/Engine/Source/Runtime/WebBrowser/Private/WebBrowserWindow.h
@@ -48,37 +48,24 @@ class FWebBrowserWindow
 	// Allow the Handler to access functions only it needs
 	friend class FWebBrowserHandler;
 
-public:
+	// The WebBrowserSingleton should be the only one creating instances of this class
+	friend class FWebBrowserSingleton;
+
+private:
 	/**
 	 * Creates and initializes a new instance.
 	 * 
-	 * @param InViewportSize Initial size of the browser window.
-	 * @param InInitialURL The Initial URL that will be loaded.
+	 * @param InBrowser The CefBrowser object representing this browser window.
+	 * @param InUrl The Initial URL that will be loaded.
 	 * @param InContentsToLoad Optional string to load as a web page.
 	 * @param InShowErrorMessage Whether to show an error message in case of loading errors.
 	 */
-	FWebBrowserWindow(FIntPoint ViewportSize, FString URL, TOptional<FString> ContentsToLoad, bool ShowErrorMessage, bool bThumbMouseButtonNavigation, bool bUseTransparency);
+	FWebBrowserWindow(CefRefPtr<CefBrowser> InBrowser, FString InUrl, TOptional<FString> InContentsToLoad, bool ShowErrorMessage, bool bThumbMouseButtonNavigation, bool bUseTransparency);
 
+public:
 	/** Virtual Destructor. */
 	virtual ~FWebBrowserWindow();
 
-public:
-
-	/**
-	* Set the CEF Handler receiving browser callbacks for this window.
-	*
-	* @param InHandler Pointer to the handler for this window.
-	*/
-	void SetHandler(CefRefPtr<FWebBrowserHandler> InHandler);
-
-	/**
-	 * Called to pass reference to the underlying CefBrowser as this is not created at the same time
-	 * as the FWebBrowserWindow.
-	 *
-	 * @param Browser The CefBrowser for this window.
-	 */
-	void BindCefBrowser(CefRefPtr<CefBrowser> Browser);
-
 	bool IsShowingErrorMessages() { return ShowErrorMessage; }
 	bool IsThumbMouseButtonNavigationEnabled() { return bThumbMouseButtonNavigation; }
 	bool UseTransparency() { return bUseTransparency; }
@@ -392,9 +379,6 @@ private:
 	/** Interface to the texture we are rendering to. */
 	FSlateUpdatableTexture* UpdatableTextures[2];
 
-	/** Pointer to the CEF Handler for this window. */
-	CefRefPtr<FWebBrowserHandler> Handler;
-
 	/** Pointer to the CEF Browser for this window. */
 	CefRefPtr<CefBrowser> InternalCefBrowser;
 
diff --git a/Engine/Source/Runtime/WebBrowser/Private/WebJSFunction.cpp b/Engine/Source/Runtime/WebBrowser/Private/WebJSFunction.cpp
--- a/Engine/Source/Runtime/WebBrowser/Private/WebJSFunction.cpp
+++ b/Engine/Source/Runtime/WebBrowser/Private/WebJSFunction.cpp
@@ -19,15 +19,21 @@
 
 FWebJSParam::~FWebJSParam()
 {
-	// Since the FString and TArray members are in a union, they may or may not be valid, so we have to call the destructors manually.
+	// Since the FString, StructWrapper, TArray, and TMap members are in a union, they may or may not be valid, so we have to call the destructors manually.
 	switch (Tag)
 	{
 		case PTYPE_STRING:
 			delete StringValue;
 			break;
+		case PTYPE_STRUCT:
+			delete StructValue;
+			break;
 		case PTYPE_ARRAY:
 			delete ArrayValue;
 			break;
+		case PTYPE_MAP:
+			delete MapValue;
+			break;
 		default:
 			break;
 	}
@@ -56,12 +62,14 @@ FWebJSParam::FWebJSParam(const FWebJSParam& Other)
 			ObjectValue = Other.ObjectValue;
 			break;
 		case PTYPE_STRUCT:
-			StructValue.TypeInfo = Other.StructValue.TypeInfo;
-			StructValue.StructPtr = Other.StructValue.StructPtr;
+			StructValue = Other.StructValue->Clone();
 			break;
 		case PTYPE_ARRAY:
 			ArrayValue = new TArray<FWebJSParam>(*Other.ArrayValue);
 			break;
+		case PTYPE_MAP:
+			MapValue = new TMap<FString, FWebJSParam>(*Other.MapValue);
+			break;
 	}
 }
 
@@ -74,4 +82,15 @@ void FWebJSFunction::Invoke(int32 ArgCount, FWebJSParam Arguments[]) const
 		Scripting->InvokeJSFunction(FunctionId, ArgCount, Arguments);
 	}
 #endif
+}
+
+void FWebJSResponse::Invoke(int32 ArgCount, FWebJSParam Arguments[], bool bIsError) const
+{
+#if WITH_CEF3
+	TSharedPtr<FWebJSScripting> Scripting = ScriptingPtr.Pin();
+	if (Scripting.IsValid())
+	{
+		Scripting->InvokeJSFunction(CallbackId, ArgCount, Arguments, bIsError);
+	}
+#endif
 }
\ No newline at end of file
diff --git a/Engine/Source/Runtime/WebBrowser/Private/WebJSScripting.cpp b/Engine/Source/Runtime/WebBrowser/Private/WebJSScripting.cpp
--- a/Engine/Source/Runtime/WebBrowser/Private/WebJSScripting.cpp
+++ b/Engine/Source/Runtime/WebBrowser/Private/WebJSScripting.cpp
@@ -244,10 +244,12 @@ bool FWebJSScripting::HandleExecuteUObjectMethodMessage(CefRefPtr<CefListValue>
 	uint16 ParamsSize = Function->ParmsSize;
 	TArray<uint8> Params;
 	UProperty* ReturnParam = nullptr;
+	UProperty* PromiseParam = nullptr;
 
 	// Convert cef argument list to a dictionary, so we can use FStructDeserializer to convert it for us
-	CefRefPtr<CefDictionaryValue> NamedArgs = CefDictionaryValue::Create();
+	if (ParamsSize > 0)
 	{
+		CefRefPtr<CefDictionaryValue> NamedArgs = CefDictionaryValue::Create();
 		int32 CurrentArg = 0;
 		CefRefPtr<CefListValue> CefArgs = MessageArguments->GetList(3);
 		for ( TFieldIterator<UProperty> It(Function); It; ++It )
@@ -261,49 +263,60 @@ bool FWebJSScripting::HandleExecuteUObjectMethodMessage(CefRefPtr<CefListValue>
 				}
 				else
 				{
-					CopyContainerValue(NamedArgs, CefArgs, CefString(*Param->GetName()), CurrentArg);
-					CurrentArg++;
+					UStructProperty *StructProperty = Cast<UStructProperty>(Param);
+					if (StructProperty && StructProperty->Struct->IsChildOf(FWebJSResponse::StaticStruct()))
+					{
+						PromiseParam = Param;
+					}
+					else
+					{
+						CopyContainerValue(NamedArgs, CefArgs, CefString(*Param->GetName()), CurrentArg);
+						CurrentArg++;
+					}
 				}
 			}
 		}
 
-		if (ParamsSize > 0)
-		{
-			// UFunction is a subclass of UStruct, so we can treat the arguments as a struct for deserialization
+		// UFunction is a subclass of UStruct, so we can treat the arguments as a struct for deserialization
+		Params.AddUninitialized(ParamsSize);
+		Function->InitializeStruct(Params.GetData());
+		FWebJSStructDeserializerBackend Backend = FWebJSStructDeserializerBackend(SharedThis(this), NamedArgs);
+		FStructDeserializer::Deserialize(Params.GetData(), *Function, Backend);
+	}
 
-			Params.AddUninitialized(ParamsSize);
-			Function->InitializeStruct(Params.GetData());
-			FWebJSStructDeserializerBackend Backend = FWebJSStructDeserializerBackend(SharedThis(this), NamedArgs);
-			FStructDeserializer::Deserialize(Params.GetData(), *Function, Backend);
+	if (PromiseParam)
+	{
+		FWebJSResponse* PromisePtr = PromiseParam->ContainerPtrToValuePtr<FWebJSResponse>(Params.GetData());
+		if (PromisePtr)
+		{
+			*PromisePtr = FWebJSResponse(SharedThis(this), ResultCallbackId);
 		}
 	}
 
 	Object->ProcessEvent(Function, Params.GetData());
 	CefRefPtr<CefListValue> Results = CefListValue::Create();
-	if ( ReturnParam )
+
+	if ( ! PromiseParam ) // If PromiseParam is set, we assume that the UFunction will ensure it is called with the result
 	{
-		FStructSerializerPolicies ReturnPolicies;
-		ReturnPolicies.PropertyFilter = [&](const UProperty* CandidateProperty, const UProperty* ParentProperty)
+		if ( ReturnParam )
 		{
-			return ParentProperty != nullptr || CandidateProperty == ReturnParam;
-		};
-		FWebJSStructSerializerBackend ReturnBackend(SharedThis(this));
-		FStructSerializer::Serialize(Params.GetData(), *Function, ReturnBackend, ReturnPolicies);
-		CefRefPtr<CefDictionaryValue> ResultDict = ReturnBackend.GetResult();
-
-		// Extract the single return value from the serialized dictionary to an array
-		CopyContainerValue(Results, ResultDict, 0, CefString(*ReturnParam->GetName()));
+			FStructSerializerPolicies ReturnPolicies;
+			ReturnPolicies.PropertyFilter = [&](const UProperty* CandidateProperty, const UProperty* ParentProperty)
+			{
+				return ParentProperty != nullptr || CandidateProperty == ReturnParam;
+			};
+			FWebJSStructSerializerBackend ReturnBackend(SharedThis(this));
+			FStructSerializer::Serialize(Params.GetData(), *Function, ReturnBackend, ReturnPolicies);
+			CefRefPtr<CefDictionaryValue> ResultDict = ReturnBackend.GetResult();
+
+			// Extract the single return value from the serialized dictionary to an array
+			CopyContainerValue(Results, ResultDict, 0, CefString(*ReturnParam->GetName()));
+		}
+		InvokeJSFunction(ResultCallbackId, Results, false);
 	}
-	InvokeJSFunction(ResultCallbackId, Results, false);
 	return true;
 }
 
-void FWebJSScripting::BindCefBrowser(CefRefPtr<CefBrowser> Browser)
-{
-	check(Browser.get() == nullptr || InternalCefBrowser.get() == nullptr || InternalCefBrowser->IsSame(Browser));
-	InternalCefBrowser = Browser;
-}
-
 void FWebJSScripting::UnbindCefBrowser()
 {
 	InternalCefBrowser = nullptr;
diff --git a/Engine/Source/Runtime/WebBrowser/Private/WebJSScripting.h b/Engine/Source/Runtime/WebBrowser/Private/WebJSScripting.h
--- a/Engine/Source/Runtime/WebBrowser/Private/WebJSScripting.h
+++ b/Engine/Source/Runtime/WebBrowser/Private/WebJSScripting.h
@@ -23,10 +23,13 @@ class FWebJSScripting
 	, public TSharedFromThis<FWebJSScripting>
 {
 public:
-	FWebJSScripting()
+	FWebJSScripting(CefRefPtr<CefBrowser> Browser)
 		: BaseGuid(FGuid::NewGuid())
+		, InternalCefBrowser(Browser)
 	{}
 
+	void UnbindCefBrowser();
+
 	void BindUObject(const FString& Name, UObject* Object, bool bIsPermanent = true);
 	void UnbindUObject(const FString& Name, UObject* Object = nullptr, bool bIsPermanent = true);
 
@@ -44,7 +47,7 @@ public:
 	 * Sends a message to the renderer process. 
 	 * See [REDACTED_PROFILE_URL]/cef/wiki/GeneralUsage#markdown-header-inter-process-communication-ipc for more information.
 	 *
-	 * @paran Message the message to send to the renderer process
+	 * @param Message the message to send to the renderer process
 	 */
 	void SendProcessMessage(CefRefPtr<CefProcessMessage> Message);
 
@@ -85,7 +88,7 @@ public:
 			}
 			case FWebJSParam::PTYPE_STRUCT:
 			{
-				CefRefPtr<CefDictionaryValue> ConvertedStruct = ConvertStruct(Param.StructValue.TypeInfo, Param.StructValue.StructPtr);
+				CefRefPtr<CefDictionaryValue> ConvertedStruct = ConvertStruct(Param.StructValue->GetTypeInfo(), Param.StructValue->GetData());
 				return Container->SetDictionary(Key, ConvertedStruct);
 			}
 			case FWebJSParam::PTYPE_ARRAY:
@@ -97,13 +100,20 @@ public:
 				}
 				return Container->SetList(Key, ConvertedArray);
 			}
+			case FWebJSParam::PTYPE_MAP:
+			{
+				CefRefPtr<CefDictionaryValue> ConvertedMap = CefDictionaryValue::Create();
+				for(auto& Pair : *Param.MapValue)
+				{
+					SetConverted(ConvertedMap, *Pair.Key, Pair.Value);
+				}
+				return Container->SetDictionary(Key, ConvertedMap);
+			}
 			default:
 				return false;
 		}
 	}
 
-	void BindCefBrowser(CefRefPtr<CefBrowser> Browser);
-	void UnbindCefBrowser();
 	CefRefPtr<CefDictionaryValue> GetPermanentBindings();
 
 	void InvokeJSFunction(FGuid FunctionId, int32 ArgCount, FWebJSParam Arguments[], bool bIsError=false);
diff --git a/Engine/Source/Runtime/WebBrowser/Public/IWebBrowserSingleton.h b/Engine/Source/Runtime/WebBrowser/Public/IWebBrowserSingleton.h
--- a/Engine/Source/Runtime/WebBrowser/Public/IWebBrowserSingleton.h
+++ b/Engine/Source/Runtime/WebBrowser/Public/IWebBrowserSingleton.h
@@ -36,8 +36,6 @@ public:
 	 *
 	 * @param OSWindowHandle Handle of OS Window that the browser will be displayed in (can be null)
 	 * @param InitialURL URL that the browser should initially navigate to
-	 * @param Width Initial width of the browser
-	 * @param Height Initial height of the browser
 	 * @param bUseTransparency Whether to allow transparent rendering of pages
 	 * @param ContentsToLoad Optional string to load as a web page
 	 * @param ShowErrorMessage Whether to show an error message in case of loading errors.
@@ -47,9 +45,7 @@ public:
 	virtual TSharedPtr<IWebBrowserWindow> CreateBrowserWindow(
 		void* OSWindowHandle, 
 		FString InitialURL, 
-		uint32 Width, 
-		uint32 Height, 
-		bool bUseTransparency, 
+		bool bUseTransparency,
 		bool bThumbMouseButtonNavigation,
 		TOptional<FString> ContentsToLoad = TOptional<FString>(), 
 		bool ShowErrorMessage = true,
diff --git a/Engine/Source/Runtime/WebBrowser/Public/WebBrowserViewport.h b/Engine/Source/Runtime/WebBrowser/Public/WebBrowserViewport.h
--- a/Engine/Source/Runtime/WebBrowser/Public/WebBrowserViewport.h
+++ b/Engine/Source/Runtime/WebBrowser/Public/WebBrowserViewport.h
@@ -20,9 +20,8 @@ public:
 	 * @param InViewportWidget The Widget displaying this viewport (needed to capture mouse)
 	 * @param InIsPopup Used to initialize a viewport for showing browser popup menus instead of the main window.
 	 */
-	FWebBrowserViewport(TSharedPtr<IWebBrowserWindow> InWebBrowserWindow, TSharedPtr<SWidget> InViewportWidget, bool InIsPopup = false)
+	FWebBrowserViewport(TSharedPtr<IWebBrowserWindow> InWebBrowserWindow, bool InIsPopup = false)
 		: WebBrowserWindow(InWebBrowserWindow)
-		, ViewportWidget(InViewportWidget)
 		, bIsPopup(InIsPopup)
 	{ }
 
@@ -56,8 +55,6 @@ public:
 private:
 	/** The web browser this viewport will display */
 	TSharedPtr<IWebBrowserWindow>	WebBrowserWindow;
-	/** The viewport widget using this interface */
-	TWeakPtr<SWidget>				ViewportWidget;
 	/** Whether this viewport is showing the browser window or a popup menu widget */
 	bool const						bIsPopup;
 };
diff --git a/Engine/Source/Runtime/WebBrowser/Public/WebJSFunction.h b/Engine/Source/Runtime/WebBrowser/Public/WebJSFunction.h
--- a/Engine/Source/Runtime/WebBrowser/Public/WebJSFunction.h
+++ b/Engine/Source/Runtime/WebBrowser/Public/WebJSFunction.h
@@ -9,19 +9,36 @@ class FWebJSScripting;
 
 struct WEBBROWSER_API FWebJSParam
 {
-	enum { PTYPE_NULL, PTYPE_BOOL, PTYPE_INT, PTYPE_DOUBLE, PTYPE_STRING, PTYPE_OBJECT, PTYPE_STRUCT, PTYPE_ARRAY } Tag;
-	union
+
+	struct IStructWrapper
 	{
-		bool BoolValue;
-		double DoubleValue;
-		int32 IntValue;
-		UObject* ObjectValue;
-		const FString* StringValue;
-		struct {
-			const void* StructPtr;
-			UStruct* TypeInfo;
-		} StructValue;
-		TArray<FWebJSParam>* ArrayValue;
+		virtual ~IStructWrapper() {};
+		virtual UStruct* GetTypeInfo() = 0;
+		virtual const void* GetData() = 0;
+		virtual IStructWrapper* Clone() = 0;
+	};
+
+	template <typename T> struct FStructWrapper
+		: public IStructWrapper
+	{
+		T StructValue;
+		FStructWrapper(const T& InValue)
+			: StructValue(InValue)
+		{}
+		virtual ~FStructWrapper()
+		{}
+		virtual UStruct* GetTypeInfo() override
+		{
+			return T::StaticStruct();
+		}
+		virtual const void* GetData() override
+		{
+			return &StructValue;
+		}
+		virtual IStructWrapper* Clone() override
+		{
+			return new FStructWrapper<T>(StructValue);
+		}
 	};
 
 	FWebJSParam() : Tag(PTYPE_NULL) {}
@@ -37,13 +54,15 @@ struct WEBBROWSER_API FWebJSParam
 	FWebJSParam(double Value) : Tag(PTYPE_DOUBLE), DoubleValue(Value) {}
 	FWebJSParam(float Value) : Tag(PTYPE_DOUBLE), DoubleValue(Value) {}
 	FWebJSParam(const FString& Value) : Tag(PTYPE_STRING), StringValue(new FString(Value)) {}
+	FWebJSParam(const FText& Value) : Tag(PTYPE_STRING), StringValue(new FString(Value.ToString())) {}
+	FWebJSParam(const FName& Value) : Tag(PTYPE_STRING), StringValue(new FString(Value.ToString())) {}
 	FWebJSParam(const TCHAR* Value) : Tag(PTYPE_STRING), StringValue(new FString(Value)) {}
 	FWebJSParam(UObject* Value) : Tag(PTYPE_OBJECT), ObjectValue(Value) {}
-	FWebJSParam(const void* Value, UStruct* Struct) : Tag(PTYPE_STRUCT), StructValue({Value, Struct}) {}
 	template <typename T> FWebJSParam(const T& Value,
 		typename TEnableIf<!TIsPointerType<T>::Value, UStruct>::Type* InTypeInfo=T::StaticStruct())
 		: Tag(PTYPE_STRUCT)
-		, StructValue({&Value, InTypeInfo}) {}
+		, StructValue(new FStructWrapper<T>(Value))
+	{}
 	template <typename T> FWebJSParam(const TArray<T>& Value)
 		: Tag(PTYPE_ARRAY)
 	{
@@ -54,11 +73,49 @@ struct WEBBROWSER_API FWebJSParam
 			ArrayValue->Add(FWebJSParam(Item));
 		}
 	}
+	template <typename T> FWebJSParam(const TMap<FString, T>& Value)
+		: Tag(PTYPE_MAP)
+	{
+		MapValue = new TMap<FString, FWebJSParam>();
+		MapValue->Reserve(Value.Num());
+		for(const auto& Pair : Value)
+		{
+			MapValue->Add(Pair.Key, FWebJSParam(Pair.Value));
+		}
+	}
+	template <typename K, typename T> FWebJSParam(const TMap<K, T>& Value)
+		: Tag(PTYPE_MAP)
+	{
+		MapValue = new TMap<FString, FWebJSParam>();
+		MapValue->Reserve(Value.Num());
+		for(const auto& Pair : Value)
+		{
+			MapValue->Add(Pair.Key.ToString(), FWebJSParam(Pair.Value));
+		}
+	}
 	FWebJSParam(const FWebJSParam& Other);
 	~FWebJSParam();
 
+	enum { PTYPE_NULL, PTYPE_BOOL, PTYPE_INT, PTYPE_DOUBLE, PTYPE_STRING, PTYPE_OBJECT, PTYPE_STRUCT, PTYPE_ARRAY, PTYPE_MAP } Tag;
+	union
+	{
+		bool BoolValue;
+		double DoubleValue;
+		int32 IntValue;
+		UObject* ObjectValue;
+		const FString* StringValue;
+		IStructWrapper* StructValue;
+		TArray<FWebJSParam>* ArrayValue;
+		TMap<FString, FWebJSParam>* MapValue;
+	};
+
 };
 
+/**
+ * Representation of a remote JS function.
+ * FWebJSFunction objects represent a JS function and allow calling them from native code.
+ * FWebJSFunction objects can also be added to delegates and events using the Bind/AddLambda method.
+ */
 USTRUCT()
 struct WEBBROWSER_API FWebJSFunction
 {
@@ -142,4 +199,68 @@ private:
 
 	TWeakPtr<FWebJSScripting> ScriptingPtr;
 	FGuid FunctionId;
-};
\ No newline at end of file
+};
+
+/** 
+ *  Representation of a remote JS async response object.
+ *  UFUNCTIONs taking a FWebJSResponse will get it passed in automatically when called from a web browser.
+ *  Pass a result or error back by invoking Success or Failure on the object.
+ *  UFunctions accepting a FWebJSResponse should have a void return type, as any value returned from the function will be ignored.
+ *  Calling the response methods does not have to happen before returning from the function, which means you can use this to implement asynchronous functionality.
+ *
+ *  Note that the remote object will become invalid as soon as a result has been delivered, so you can only call either Success or Failure once.
+ */
+USTRUCT()
+struct WEBBROWSER_API FWebJSResponse
+{
+	GENERATED_USTRUCT_BODY()
+
+	FWebJSResponse()
+		: ScriptingPtr()
+		, CallbackId()
+	{}
+
+	FWebJSResponse(TSharedPtr<FWebJSScripting> InScripting, const FGuid& InCallbackId)
+		: ScriptingPtr(InScripting)
+		, CallbackId(InCallbackId)
+	{}
+
+	/**
+	 * Indicate successful completion without a return value.
+	 * The remote Promise's then() handler will be executed without arguments.
+	 */
+	void Success() const
+	{
+		Invoke(0, nullptr, false);
+	}
+
+	/**
+	 * Indicate successful completion passing a return value back.
+	 * The remote Promise's then() handler will be executed with the value passed as its single argument.
+	 */
+	template<typename T>
+	void Success(T Arg) const
+	{
+		FWebJSParam ArgArray[1] = {FWebJSParam(Arg)};
+		Invoke(1, ArgArray, false);
+	}
+
+	/**
+	 * Indicate failed completion, passing an error message back to JS.
+	 * The remote Promise's catch() handler will be executed with the value passed as the error reason.
+	 */
+	template<typename T>
+	void Failure(T Arg) const
+	{
+		FWebJSParam ArgArray[1] = {FWebJSParam(Arg)};
+		Invoke(1, ArgArray, true);
+	}
+
+private:
+
+	void Invoke(int32 ArgCount, FWebJSParam Arguments[], bool bIsError) const;
+
+	TWeakPtr<FWebJSScripting> ScriptingPtr;
+	FGuid CallbackId;
+
+};
