[INTEGRATE] SWebBrowser: Add hooks to to be able to handle alert and message boxes with a custom handler
OPP-3413
Merging CL#2637687 using UE4-To-UE4-LauncherDev

diff --git a/Engine/Source/Runtime/WebBrowser/Private/SWebBrowser.cpp b/Engine/Source/Runtime/WebBrowser/Private/SWebBrowser.cpp
--- a/Engine/Source/Runtime/WebBrowser/Private/SWebBrowser.cpp
+++ b/Engine/Source/Runtime/WebBrowser/Private/SWebBrowser.cpp
@@ -29,6 +29,8 @@ SWebBrowser::~SWebBrowser()
 		BrowserWindow->OnDismissPopup().RemoveAll(this);
 		BrowserWindow->OnBeforeBrowse().Unbind();
 		BrowserWindow->OnLoadUrl().Unbind();
+		BrowserWindow->OnShowDialog().Unbind();
+		BrowserWindow->OnDismissAllDialogs().Unbind();
 		BrowserWindow->OnBeforePopup().Unbind();
 		if(!BrowserWindow->IsClosing())
 		{
@@ -46,6 +48,8 @@ void SWebBrowser::Construct(const FArguments& InArgs, const TSharedPtr<IWebBrows
 	OnUrlChanged = InArgs._OnUrlChanged;
 	OnBeforeNavigation = InArgs._OnBeforeNavigation;
 	OnLoadUrl = InArgs._OnLoadUrl;
+	OnShowDialog = InArgs._OnShowDialog;
+	OnDismissAllDialogs = InArgs._OnDismissAllDialogs;
 	OnBeforePopup = InArgs._OnBeforePopup;
 	OnCreateWindow = InArgs._OnCreateWindow;
 	OnCloseWindow = InArgs._OnCloseWindow;
@@ -184,6 +188,8 @@ void SWebBrowser::Construct(const FArguments& InArgs, const TSharedPtr<IWebBrows
 		BrowserWindow->OnToolTip().AddSP(this, &SWebBrowser::HandleToolTip);
 		BrowserWindow->OnBeforeBrowse().BindSP(this, &SWebBrowser::HandleBeforeNavigation);
 		BrowserWindow->OnLoadUrl().BindSP(this, &SWebBrowser::HandleLoadUrl);
+		BrowserWindow->OnShowDialog().BindSP(this, &SWebBrowser::HandleShowDialog);
+		BrowserWindow->OnDismissAllDialogs().BindSP(this, &SWebBrowser::HandleDismissAllDialogs);
 		BrowserWindow->OnBeforePopup().BindSP(this, &SWebBrowser::HandleBeforePopup);
 		BrowserWindow->OnShowPopup().AddSP(this, &SWebBrowser::HandleShowPopup);
 		BrowserWindow->OnDismissPopup().AddSP(this, &SWebBrowser::HandleDismissPopup);
@@ -443,6 +449,21 @@ bool SWebBrowser::HandleLoadUrl(const FString& Method, const FString& Url, FStri
 	return false;
 }
 
+EWebBrowserDialogEventResponse SWebBrowser::HandleShowDialog(const TWeakPtr<IWebBrowserDialog>& DialogParams)
+{
+	if(OnShowDialog.IsBound())
+	{
+		return OnShowDialog.Execute(DialogParams);
+	}
+	return EWebBrowserDialogEventResponse::Unhandled;
+}
+
+void SWebBrowser::HandleDismissAllDialogs()
+{
+	OnDismissAllDialogs.ExecuteIfBound();
+}
+
+
 bool SWebBrowser::HandleBeforePopup(FString URL, FString Target)
 {
 	if (OnBeforePopup.IsBound())
diff --git a/Engine/Source/Runtime/WebBrowser/Private/WebBrowserDialog.h b/Engine/Source/Runtime/WebBrowser/Private/WebBrowserDialog.h
new file mode 100644
--- /dev/null
+++ b/Engine/Source/Runtime/WebBrowser/Private/WebBrowserDialog.h
@@ -0,0 +1,90 @@
+// Copyright 1998-2015 Epic Games, Inc. All Rights Reserved.
+
+#pragma once
+
+#if WITH_CEF3
+#	if PLATFORM_WINDOWS
+#		include "AllowWindowsPlatformTypes.h"
+#	endif
+#	pragma push_macro("OVERRIDE")
+#		undef OVERRIDE // cef headers provide their own OVERRIDE macro
+#		include "include/cef_jsdialog_handler.h"
+#	pragma pop_macro("OVERRIDE")
+#	if PLATFORM_WINDOWS
+#		include "HideWindowsPlatformTypes.h"
+#	endif
+#endif
+
+#include "IWebBrowserDialog.h"
+
+class FWebBrowserWindow;
+
+class FWebBrowserDialog
+	: public IWebBrowserDialog
+{
+public:
+	virtual ~FWebBrowserDialog()
+	{}
+
+	// IWebBrowserDialog interface:
+
+	virtual EWebBrowserDialogType GetType() override
+	{
+		return Type;
+	}
+
+	virtual const FText& GetMessageText() override
+	{
+		return MessageText;
+	}
+
+	virtual const FText& GetDefaultPrompt() override
+	{
+		return DefaultPrompt;
+	}
+
+	virtual bool IsReload() override
+	{
+		check(Type == EWebBrowserDialogType::Unload);
+		return bIsReload;
+	}
+
+	virtual void Continue(bool Success = true, const FText& UserResponse = FText::GetEmpty()) override
+	{
+#if WITH_CEF3
+		check(Type == EWebBrowserDialogType::Prompt || UserResponse.IsEmpty());
+		Callback->Continue(Success, *UserResponse.ToString());
+#endif
+	}
+
+private:
+
+
+	EWebBrowserDialogType Type;
+	FText MessageText;
+	FText DefaultPrompt;
+	bool bIsReload;
+#if WITH_CEF3
+	CefRefPtr<CefJSDialogCallback> Callback;
+
+	// Create a dialog from OnJSDialog arguments
+	FWebBrowserDialog(CefJSDialogHandler::JSDialogType InDialogType, const CefString& InMessageText, const CefString& InDefaultPrompt, CefRefPtr<CefJSDialogCallback> InCallback)
+		: Type((EWebBrowserDialogType)InDialogType)
+		, MessageText(FText::FromString(InMessageText.ToWString().c_str()))
+		, DefaultPrompt(FText::FromString(InDefaultPrompt.ToWString().c_str()))
+		, bIsReload(false)
+		, Callback(InCallback)
+	{}
+
+	// Create a dialog from OnBeforeUnloadDialog arguments
+	FWebBrowserDialog(const CefString& InMessageText, bool InIsReload, CefRefPtr<CefJSDialogCallback> InCallback)
+		: Type(EWebBrowserDialogType::Unload)
+		, MessageText(FText::FromString(InMessageText.ToWString().c_str()))
+		, DefaultPrompt(FText::GetEmpty())
+		, bIsReload(InIsReload)
+		, Callback(InCallback)
+	{};
+
+	friend class FWebBrowserWindow;
+#endif
+};
diff --git a/Engine/Source/Runtime/WebBrowser/Private/WebBrowserHandler.cpp b/Engine/Source/Runtime/WebBrowser/Private/WebBrowserHandler.cpp
--- a/Engine/Source/Runtime/WebBrowser/Private/WebBrowserHandler.cpp
+++ b/Engine/Source/Runtime/WebBrowser/Private/WebBrowserHandler.cpp
@@ -454,6 +454,40 @@ bool FWebBrowserHandler::OnKeyEvent(CefRefPtr<CefBrowser> Browser,
 
 	return false;
 }
-#endif
+
+bool FWebBrowserHandler::OnJSDialog(CefRefPtr<CefBrowser> Browser, const CefString& OriginUrl, const CefString& AcceptLang, JSDialogType DialogType, const CefString& MessageText, const CefString& DefaultPromptText, CefRefPtr<CefJSDialogCallback> Callback, bool& OutSuppressMessage)
+{
+	bool Retval = false;
+	TSharedPtr<FWebBrowserWindow> BrowserWindow = BrowserWindowPtr.Pin();
+	if (BrowserWindow.IsValid())
+	{
+		Retval = BrowserWindow->OnJSDialog(DialogType, MessageText, DefaultPromptText, Callback, OutSuppressMessage);
+	}
+	return Retval;
+}
+
+bool FWebBrowserHandler::OnBeforeUnloadDialog(CefRefPtr<CefBrowser> Browser, const CefString& MessageText, bool IsReload, CefRefPtr<CefJSDialogCallback> Callback)
+{
+	bool Retval = false;
+	TSharedPtr<FWebBrowserWindow> BrowserWindow = BrowserWindowPtr.Pin();
+	if (BrowserWindow.IsValid())
+	{
+		Retval = BrowserWindow->OnBeforeUnloadDialog(MessageText, IsReload, Callback);
+	}
+	return Retval;
+}
+
+void FWebBrowserHandler::OnResetDialogState(CefRefPtr<CefBrowser> Browser)
+{
+	TSharedPtr<FWebBrowserWindow> BrowserWindow = BrowserWindowPtr.Pin();
+	if (BrowserWindow.IsValid())
+	{
+		BrowserWindow->OnResetDialogState();
+	}
+}
+
+
+
+#endif // WITH_CEF
 
 #undef LOCTEXT_NAMESPACE
diff --git a/Engine/Source/Runtime/WebBrowser/Private/WebBrowserHandler.h b/Engine/Source/Runtime/WebBrowser/Private/WebBrowserHandler.h
--- a/Engine/Source/Runtime/WebBrowser/Private/WebBrowserHandler.h
+++ b/Engine/Source/Runtime/WebBrowser/Private/WebBrowserHandler.h
@@ -33,6 +33,7 @@ class FWebBrowserHandler
 	, public CefRenderHandler
 	, public CefRequestHandler
 	, public CefKeyboardHandler
+	, public CefJSDialogHandler
 {
 public:
 
@@ -106,6 +107,11 @@ public:
 		return this;
 	}
 
+	virtual CefRefPtr<CefJSDialogHandler> GetJSDialogHandler() override
+	{
+		return this;
+	}
+
 	virtual bool OnProcessMessageReceived(CefRefPtr<CefBrowser> Browser,
 		CefProcessId SourceProcess,
 		CefRefPtr<CefProcessMessage> Message) override;
@@ -192,6 +198,15 @@ public:
 		const CefKeyEvent& Event,
 		CefEventHandle OsEvent) override;
 
+public:
+	// CefJSDialogHandler interface
+
+	virtual bool OnJSDialog(CefRefPtr<CefBrowser> Browser, const CefString& OriginUrl, const CefString& AcceptLang, JSDialogType DialogType, const CefString& MessageText, const CefString& DefaultPromptText, CefRefPtr<CefJSDialogCallback> Callback, bool& OutSuppressMessage) override;
+
+	virtual bool OnBeforeUnloadDialog(CefRefPtr<CefBrowser> Browser, const CefString& MessageText, bool IsReload, CefRefPtr<CefJSDialogCallback> Callback) override;
+
+	virtual void OnResetDialogState(CefRefPtr<CefBrowser> Browser) override;
+
 private:
 
 	bool ShowDevTools(const CefRefPtr<CefBrowser>& Browser);
diff --git a/Engine/Source/Runtime/WebBrowser/Private/WebBrowserWindow.cpp b/Engine/Source/Runtime/WebBrowser/Private/WebBrowserWindow.cpp
--- a/Engine/Source/Runtime/WebBrowser/Private/WebBrowserWindow.cpp
+++ b/Engine/Source/Runtime/WebBrowser/Private/WebBrowserWindow.cpp
@@ -4,6 +4,7 @@
 #include "WebBrowserWindow.h"
 #include "WebBrowserByteResource.h"
 #include "WebBrowserPopupFeatures.h"
+#include "WebBrowserDialog.h"
 #include "WebJSScripting.h"
 #include "RHI.h"
 
@@ -456,6 +457,79 @@ bool FWebBrowserWindow::RequestCreateWindow( const TSharedRef<IWebBrowserWindow>
 	return false;
 }
 
+bool FWebBrowserWindow::OnJSDialog(CefJSDialogHandler::JSDialogType DialogType, const CefString& MessageText, const CefString& DefaultPromptText, CefRefPtr<CefJSDialogCallback> Callback, bool& OutSuppressMessage)
+{
+	bool Retval = false;
+	if ( OnShowDialog().IsBound() )
+	{
+		TSharedPtr<IWebBrowserDialog> Dialog(new FWebBrowserDialog(DialogType, MessageText, DefaultPromptText, Callback));
+		EWebBrowserDialogEventResponse EventResponse = OnShowDialog().Execute(TWeakPtr<IWebBrowserDialog>(Dialog));
+		switch (EventResponse)
+		{
+		case EWebBrowserDialogEventResponse::Handled:
+			Retval = true;
+			break;
+		case EWebBrowserDialogEventResponse::Continue:
+			if (DialogType == JSDIALOGTYPE_ALERT)
+			{
+				// Alert dialogs don't return a value, so treat Continue the same way as Ingore
+				OutSuppressMessage = true;
+				Retval = false;
+			}
+			else
+			{
+				Callback->Continue(true, DefaultPromptText);
+				Retval = true;
+			}
+			break;
+		case EWebBrowserDialogEventResponse::Ignore:
+			OutSuppressMessage = true;
+			Retval = false;
+			break;
+		case EWebBrowserDialogEventResponse::Unhandled:
+		default:
+			Retval = false;
+			break;
+		}
+	}
+	return Retval;
+}
+
+bool FWebBrowserWindow::OnBeforeUnloadDialog(const CefString& MessageText, bool IsReload, CefRefPtr<CefJSDialogCallback> Callback)
+{
+	bool Retval = false;
+	if ( OnShowDialog().IsBound() )
+	{
+		TSharedPtr<IWebBrowserDialog> Dialog(new FWebBrowserDialog(MessageText, IsReload, Callback));
+		EWebBrowserDialogEventResponse EventResponse = OnShowDialog().Execute(TWeakPtr<IWebBrowserDialog>(Dialog));
+		switch (EventResponse)
+		{
+		case EWebBrowserDialogEventResponse::Handled:
+			Retval = true;
+			break;
+		case EWebBrowserDialogEventResponse::Continue:
+			Callback->Continue(true, CefString());
+			Retval = true;
+			break;
+		case EWebBrowserDialogEventResponse::Ignore:
+			Callback->Continue(false, CefString());
+			Retval = true;
+			break;
+		case EWebBrowserDialogEventResponse::Unhandled:
+		default:
+			Retval = false;
+			break;
+		}
+	}
+	return Retval;
+}
+
+void FWebBrowserWindow::OnResetDialogState()
+{
+	OnDismissAllDialogs().ExecuteIfBound();
+}
+
+
 FReply FWebBrowserWindow::OnMouseButtonDown(const FGeometry& MyGeometry, const FPointerEvent& MouseEvent, bool bIsPopup)
 {
 	FReply Reply = FReply::Unhandled();
diff --git a/Engine/Source/Runtime/WebBrowser/Private/WebBrowserWindow.h b/Engine/Source/Runtime/WebBrowser/Private/WebBrowserWindow.h
--- a/Engine/Source/Runtime/WebBrowser/Private/WebBrowserWindow.h
+++ b/Engine/Source/Runtime/WebBrowser/Private/WebBrowserWindow.h
@@ -19,6 +19,7 @@
 #undef OVERRIDE // cef headers provide their own OVERRIDE macro
 #include "include/internal/cef_ptr.h"
 #include "include/cef_render_handler.h"
+#include "include/cef_jsdialog_handler.h"
 #pragma pop_macro("OVERRIDE")
 
 #if PLATFORM_WINDOWS
@@ -28,7 +29,7 @@
 /**
  * Helper for containing items required for CEF browser window creation.
  */
-struct  FWebBrowserWindowInfo
+struct FWebBrowserWindowInfo
 {
 	FWebBrowserWindowInfo(CefRefPtr<CefBrowser> InBrowser, CefRefPtr<FWebBrowserHandler> InHandler) 
 		: Browser(InBrowser)
@@ -178,6 +179,16 @@ public:
 		return DismissPopupEvent;
 	}
 
+	virtual FOnShowDialog& OnShowDialog() override
+	{
+		return ShowDialogDelegate;
+	}
+
+	virtual FOnDismissAllDialogs& OnDismissAllDialogs() override
+	{
+		return DismissAllDialogsDelegate;
+	}
+
 private:
 
 	/**
@@ -275,6 +286,21 @@ private:
 	 */
 	bool OnUnhandledKeyEvent(const CefKeyEvent& CefEvent);
 
+	/**
+	 * Handle showing javascript dialogs
+	 */
+	bool OnJSDialog(CefJSDialogHandler::JSDialogType DialogType, const CefString& MessageText, const CefString& DefaultPromptText, CefRefPtr<CefJSDialogCallback> Callback, bool& OutSuppressMessage);
+
+	/**
+	 * Handle showing unload confirmation dialogs
+	 */
+	bool OnBeforeUnloadDialog(const CefString& MessageText, bool IsReload, CefRefPtr<CefJSDialogCallback> Callback);
+
+	/**
+	 * Notify when any and all pending dialogs should be canceled
+	 */
+	void OnResetDialogState();
+
 
 	/** Specifies if window creation functionality is available. 
 	 *
@@ -436,18 +462,24 @@ private:
 	/** Delegate for notifying that a popup window is attempting to open. */
 	FOnBeforePopupDelegate BeforePopupDelegate;
 	
-	/** Delegate for handaling requests to create new windows. */
+	/** Delegate for handling requests to create new windows. */
 	FOnCreateWindow CreateWindowDelegate;
 
-	/** Delegate for handaling requests to close new windows that were created. */
+	/** Delegate for handling requests to close new windows that were created. */
 	FOnCloseWindow CloseWindowDelegate;
 
-	/** Delegate for handaling requests to show the popup menu. */
+	/** Delegate for handling requests to show the popup menu. */
 	FOnShowPopup ShowPopupEvent;
 
-	/** Delegate for hanaling requests to dismiss the current popup menu. */
+	/** Delegate for handling requests to dismiss the current popup menu. */
 	FOnDismissPopup DismissPopupEvent;
 
+	/** Delegate for showing dialogs. */
+	FOnShowDialog ShowDialogDelegate;
+
+	/** Delegate for dismissing all dialogs. */
+	FOnDismissAllDialogs DismissAllDialogsDelegate;
+
 	/** Tracks the current mouse cursor */
 	EMouseCursor::Type Cursor;
 
diff --git a/Engine/Source/Runtime/WebBrowser/Public/IWebBrowserDialog.h b/Engine/Source/Runtime/WebBrowser/Public/IWebBrowserDialog.h
new file mode 100644
--- /dev/null
+++ b/Engine/Source/Runtime/WebBrowser/Public/IWebBrowserDialog.h
@@ -0,0 +1,80 @@
+// Copyright 1998-2015 Epic Games, Inc. All Rights Reserved.
+
+#pragma once
+
+/**
+ * Type of browser dialog to show.
+ */
+
+enum class EWebBrowserDialogType
+{
+	/** Alert dialog. Used to show a message to a user. It should only have a single button to dismiss it. */
+	Alert  = 0,
+	/** Confirm dialog. Shows a message to the user. It should have two buttons to either accept or decline the prompt. */
+	Confirm,
+	/** Prompt dialog. Shows a prompt asking for user input. The user can enter text and either confirm or dismiss it. */
+	Prompt,
+	/** Unload dialog. This type of dialog is shown to confirm navigating away from a page containing user-edited content. */
+	Unload = 127,
+};
+
+/**
+ * Return value from dialog event handle specifying what action should be taken.
+ */
+enum class EWebBrowserDialogEventResponse
+{
+	/** Return Unhandled to use the default dialog implementation. This is the default behavior when no handler is attached. */
+	Unhandled,
+	/** Do not show any dialog and return as if the user accepted the action. */
+	Continue,
+	/** Do not show any dialog and return as if the user dismissed the action. */
+	Ignore,
+	/** The event handler will take care of showing the dialog. It must call IWebBrowserDialog::Continue once it has been dismissed. */
+	Handled,
+};
+
+/**
+ * Browser dialog parameters passed to OnBrowserDialog event handlers.
+ */
+class IWebBrowserDialog
+{
+public:
+
+	virtual ~IWebBrowserDialog()
+	{}
+
+	/** 
+	 * What kind of dialog should be shown.
+	 * @return the dialog type 
+	 */
+	virtual EWebBrowserDialogType GetType() = 0;
+
+	/**
+	 * Tell the browser to continue with the result of the dialog.
+	 * If this method is used, the original event handler must return EWebBrowserDialogEventResponse::Handled.
+	 *
+	 * @param Success Did the user accept the dialog or not?
+	 * @param UserResponse Only for Prompt dialog, the text entered by the user.
+	 */
+	virtual void Continue(bool Success = true, const FText& UserResponse = FText::GetEmpty()) = 0;
+
+	/**
+	 * Get the dialog message.
+	 * @return the message to show in the dialog.
+	 */
+	virtual const FText& GetMessageText() = 0;
+
+	/**
+	 * Only valid for Prompt dialogs.
+	 * @return the default value to show in the text entry box.
+	 */
+	virtual const FText& GetDefaultPrompt() = 0;
+
+	/**
+	 * Only valid for Unload dialogs.
+	 * @return true if the dialog is confirming a reload of the current page.
+	 */
+	virtual bool IsReload() = 0;
+
+	
+};
\ No newline at end of file
diff --git a/Engine/Source/Runtime/WebBrowser/Public/IWebBrowserWindow.h b/Engine/Source/Runtime/WebBrowser/Public/IWebBrowserWindow.h
--- a/Engine/Source/Runtime/WebBrowser/Public/IWebBrowserWindow.h
+++ b/Engine/Source/Runtime/WebBrowser/Public/IWebBrowserWindow.h
@@ -2,6 +2,7 @@
 
 #pragma once
 
+#include "IWebbrowserDialog.h"
 
 struct FGeometry;
 struct FKeyEvent;
@@ -12,6 +13,7 @@ class FCursorReply;
 class FSlateShaderResource;
 class IWebBrowserPopupFeatures;
 
+
 enum class EWebBrowserDocumentState
 {
 	Completed,
@@ -287,6 +289,14 @@ public:
 	DECLARE_EVENT(IWebBrowserWindow, FOnDismissPopup)
 	virtual FOnDismissPopup& OnDismissPopup() = 0;
 
+	/** A delegate that is invoked when the browser needs to show a dialog. */
+	DECLARE_DELEGATE_RetVal_OneParam(EWebBrowserDialogEventResponse, FOnShowDialog, const TWeakPtr<IWebBrowserDialog>& /*DialogParams*/)
+	virtual FOnShowDialog& OnShowDialog() = 0;
+
+	/** A delegate that is invoked when the browser needs to dismiss and reset all dialogs. */
+	DECLARE_DELEGATE(FOnDismissAllDialogs)
+	virtual FOnDismissAllDialogs& OnDismissAllDialogs() = 0;
+
 protected:
 
 	/** Virtual Destructor. */
diff --git a/Engine/Source/Runtime/WebBrowser/Public/SWebBrowser.h b/Engine/Source/Runtime/WebBrowser/Public/SWebBrowser.h
--- a/Engine/Source/Runtime/WebBrowser/Public/SWebBrowser.h
+++ b/Engine/Source/Runtime/WebBrowser/Public/SWebBrowser.h
@@ -3,6 +3,7 @@
 #pragma once
 
 #include "SlateBasics.h"
+#include "IWebBrowserDialog.h"
 
 enum class EWebBrowserDocumentState;
 class IWebBrowserWindow;
@@ -21,6 +22,7 @@ class WEBBROWSER_API SWebBrowser
 public:
 	DECLARE_DELEGATE_RetVal_TwoParams(bool, FOnBeforeBrowse, const FString&, bool)
 	DECLARE_DELEGATE_RetVal_ThreeParams(bool, FOnLoadUrl, const FString& /*Method*/, const FString& /*Url*/, FString& /* Response */)
+	DECLARE_DELEGATE_RetVal_OneParam(EWebBrowserDialogEventResponse, FOnShowDialog, const TWeakPtr<IWebBrowserDialog>&);
 
 	SLATE_BEGIN_ARGS(SWebBrowser)
 		: _InitialURL(TEXT("https://www.google.com"))
@@ -96,6 +98,13 @@ public:
 	
 		/** Called to allow bypassing page content on load. */
 		SLATE_EVENT(FOnLoadUrl, OnLoadUrl)
+
+		/** Called when the browser needs to show a dialog to the user. */
+		SLATE_EVENT(FOnShowDialog, OnShowDialog)
+
+		/** Called to dismiss any dialogs shown via OnShowDialog. */
+		SLATE_EVENT(FSimpleDelegate, OnDismissAllDialogs)
+
 	SLATE_END_ARGS()
 
 
@@ -250,6 +259,12 @@ private:
 	 */
 	bool HandleCloseWindow(const TWeakPtr<IWebBrowserWindow>& BrowserWindow);
 
+	/** Callback for showing dialogs to the user */
+	EWebBrowserDialogEventResponse HandleShowDialog(const TWeakPtr<IWebBrowserDialog>& DialogParams);
+
+	/** Callback for dismissing any dialogs previously shown  */
+	void HandleDismissAllDialogs();
+
 	/** Callback for popup window permission */
 	bool HandleBeforePopup(FString URL, FString Target);
 
@@ -324,4 +339,11 @@ private:
 	
 	/** A delegate that is invoked when loading a resource, allowing the application to provide contents directly */
 	FOnLoadUrl OnLoadUrl;
+
+	/** A delegate that is invoked when when the browser needs to present a dialog to the user */
+	FOnShowDialog OnShowDialog;
+
+	/** A delegate that is invoked when when the browser needs to dismiss all dialogs */
+	FSimpleDelegate OnDismissAllDialogs;
+
 };
