WebBrowser: Fix compilation on a case-sensitive filesystem.

diff --git a/Engine/Source/Runtime/WebBrowser/Public/IWebBrowserWindow.h b/Engine/Source/Runtime/WebBrowser/Public/IWebBrowserWindow.h
--- a/Engine/Source/Runtime/WebBrowser/Public/IWebBrowserWindow.h
+++ b/Engine/Source/Runtime/WebBrowser/Public/IWebBrowserWindow.h
@@ -2,7 +2,7 @@
 
 #pragma once
 
-#include "IWebbrowserDialog.h"
+#include "IWebBrowserDialog.h"
 
 struct FGeometry;
 struct FKeyEvent;
