Merging part of CL#2641230 & CL#2641334 LauncherDev->Main
change summary:
- Adding keyboard focus support to webbrowser widget.

diff --git a/Engine/Source/Runtime/WebBrowser/Public/SWebBrowser.h b/Engine/Source/Runtime/WebBrowser/Public/SWebBrowser.h
--- a/Engine/Source/Runtime/WebBrowser/Public/SWebBrowser.h
+++ b/Engine/Source/Runtime/WebBrowser/Public/SWebBrowser.h
@@ -113,6 +113,8 @@ public:
 
 	~SWebBrowser();
 
+	virtual bool SupportsKeyboardFocus() const override {return true;}
+
 	/**
 	 * Construct the widget.
 	 *
