[INTEGRATE] WebBrowser: Extract common functionality from FWebJSFunction and FWebJSResult into a common base class.
Added IsValid to check wether the object is valid and bound to a scripting runtime.

Merging CL#2645093 using UE4-To-UE4-LauncherDev

diff --git a/Engine/Source/Runtime/WebBrowser/Private/WebJSFunction.cpp b/Engine/Source/Runtime/WebBrowser/Private/WebJSFunction.cpp
--- a/Engine/Source/Runtime/WebBrowser/Private/WebJSFunction.cpp
+++ b/Engine/Source/Runtime/WebBrowser/Private/WebJSFunction.cpp
@@ -73,18 +73,7 @@ FWebJSParam::FWebJSParam(const FWebJSParam& Other)
 	}
 }
 
-void FWebJSFunction::Invoke(int32 ArgCount, FWebJSParam Arguments[]) const
-{
-#if WITH_CEF3
-	TSharedPtr<FWebJSScripting> Scripting = ScriptingPtr.Pin();
-	if (Scripting.IsValid())
-	{
-		Scripting->InvokeJSFunction(FunctionId, ArgCount, Arguments);
-	}
-#endif
-}
-
-void FWebJSResponse::Invoke(int32 ArgCount, FWebJSParam Arguments[], bool bIsError) const
+void FWebJSCallbackBase::Invoke(int32 ArgCount, FWebJSParam Arguments[], bool bIsError) const
 {
 #if WITH_CEF3
 	TSharedPtr<FWebJSScripting> Scripting = ScriptingPtr.Pin();
diff --git a/Engine/Source/Runtime/WebBrowser/Public/WebJSFunction.h b/Engine/Source/Runtime/WebBrowser/Public/WebJSFunction.h
--- a/Engine/Source/Runtime/WebBrowser/Public/WebJSFunction.h
+++ b/Engine/Source/Runtime/WebBrowser/Public/WebJSFunction.h
@@ -5,8 +5,6 @@
 #include "EnableIf.h"
 #include "WebJSFunction.generated.h"
 
-class FWebJSScripting;
-
 struct WEBBROWSER_API FWebJSParam
 {
 
@@ -111,6 +109,37 @@ struct WEBBROWSER_API FWebJSParam
 
 };
 
+class FWebJSScripting;
+
+/** Base class for JS callback objects. */
+USTRUCT()
+struct WEBBROWSER_API FWebJSCallbackBase
+{
+	GENERATED_USTRUCT_BODY()
+	FWebJSCallbackBase()
+	{}
+
+	bool IsValid() const
+	{
+		return ScriptingPtr.IsValid();
+	}
+
+
+protected:
+	FWebJSCallbackBase(TSharedPtr<FWebJSScripting> InScripting, const FGuid& InCallbackId)
+		: ScriptingPtr(InScripting)
+		, CallbackId(InCallbackId)
+	{}
+
+	void Invoke(int32 ArgCount, FWebJSParam Arguments[], bool bIsError = false) const;
+
+private:
+
+	TWeakPtr<FWebJSScripting> ScriptingPtr;
+	FGuid CallbackId;
+};
+
+
 /**
  * Representation of a remote JS function.
  * FWebJSFunction objects represent a JS function and allow calling them from native code.
@@ -118,17 +147,16 @@ struct WEBBROWSER_API FWebJSParam
  */
 USTRUCT()
 struct WEBBROWSER_API FWebJSFunction
+	: public FWebJSCallbackBase
 {
 	GENERATED_USTRUCT_BODY()
 
 	FWebJSFunction()
-		: ScriptingPtr()
-		, FunctionId()
+		: FWebJSCallbackBase()
 	{}
 
 	FWebJSFunction(TSharedPtr<FWebJSScripting> InScripting, const FGuid& InFunctionId)
-		: ScriptingPtr(InScripting)
-		, FunctionId(InFunctionId)
+		: FWebJSCallbackBase(InScripting, InFunctionId)
 	{}
 
 #if PLATFORM_COMPILER_HAS_VARIADIC_TEMPLATES
@@ -192,13 +220,6 @@ struct WEBBROWSER_API FWebJSFunction
 		Invoke(7, ArgArray);
 	}
 #endif
-
-private:
-
-	void Invoke(int32 ArgCount, FWebJSParam Arguments[]) const;
-
-	TWeakPtr<FWebJSScripting> ScriptingPtr;
-	FGuid FunctionId;
 };
 
 /** 
@@ -212,17 +233,16 @@ private:
  */
 USTRUCT()
 struct WEBBROWSER_API FWebJSResponse
+	: public FWebJSCallbackBase
 {
 	GENERATED_USTRUCT_BODY()
 
 	FWebJSResponse()
-		: ScriptingPtr()
-		, CallbackId()
+		: FWebJSCallbackBase()
 	{}
 
 	FWebJSResponse(TSharedPtr<FWebJSScripting> InScripting, const FGuid& InCallbackId)
-		: ScriptingPtr(InScripting)
-		, CallbackId(InCallbackId)
+		: FWebJSCallbackBase(InScripting, InCallbackId)
 	{}
 
 	/**
@@ -256,11 +276,5 @@ struct WEBBROWSER_API FWebJSResponse
 		Invoke(1, ArgArray, true);
 	}
 
-private:
-
-	void Invoke(int32 ArgCount, FWebJSParam Arguments[], bool bIsError) const;
-
-	TWeakPtr<FWebJSScripting> ScriptingPtr;
-	FGuid CallbackId;
 
 };
