Merging CL# 2675909 using UE4-To-UE4-LauncherDev

Initial web browser spinner/throbber overlay can now be disabled.

diff --git a/Engine/Source/Runtime/WebBrowser/Private/SWebBrowser.cpp b/Engine/Source/Runtime/WebBrowser/Private/SWebBrowser.cpp
--- a/Engine/Source/Runtime/WebBrowser/Private/SWebBrowser.cpp
+++ b/Engine/Source/Runtime/WebBrowser/Private/SWebBrowser.cpp
@@ -55,6 +55,7 @@ void SWebBrowser::Construct(const FArguments& InArgs, const TSharedPtr<IWebBrows
 	OnCloseWindow = InArgs._OnCloseWindow;
 	AddressBarUrl = FText::FromString(InArgs._InitialURL);
 	PopupMenuMethod = InArgs._PopupMenuMethod;
+	bShowInitialThrobber = InArgs._ShowInitialThrobber;
 
 	BrowserWindow = InWebBrowserWindow;
 	if(!BrowserWindow.IsValid())
@@ -362,7 +363,7 @@ void SWebBrowser::OnUrlTextCommitted( const FText& NewText, ETextCommit::Type Co
 
 EVisibility SWebBrowser::GetViewportVisibility() const
 {
-	if (BrowserWindow.IsValid() && BrowserWindow->IsInitialized())
+	if (!bShowInitialThrobber || (BrowserWindow.IsValid() &&  BrowserWindow->IsInitialized()))
 	{
 		return EVisibility::Visible;
 	}
@@ -371,7 +372,7 @@ EVisibility SWebBrowser::GetViewportVisibility() const
 
 EVisibility SWebBrowser::GetLoadingThrobberVisibility() const
 {
-	if (!BrowserWindow.IsValid() || !BrowserWindow->IsInitialized())
+	if (bShowInitialThrobber && (!BrowserWindow.IsValid() || !BrowserWindow->IsInitialized()))
 	{
 		return EVisibility::Visible;
 	}
diff --git a/Engine/Source/Runtime/WebBrowser/Public/SWebBrowser.h b/Engine/Source/Runtime/WebBrowser/Public/SWebBrowser.h
--- a/Engine/Source/Runtime/WebBrowser/Public/SWebBrowser.h
+++ b/Engine/Source/Runtime/WebBrowser/Public/SWebBrowser.h
@@ -31,6 +31,7 @@ public:
 		, _ShowErrorMessage(true)
 		, _SupportsTransparency(false)
 		, _SupportsThumbMouseButtonNavigation(false)
+		, _ShowInitialThrobber(true)
 		, _BackgroundColor(255,255,255,255)
 		, _PopupMenuMethod(TOptional<EPopupMethod>())
 		, _ViewportSize(FVector2D::ZeroVector)
@@ -60,6 +61,9 @@ public:
 		/** Whether to allow forward and back navigation via the mouse thumb buttons. */
 		SLATE_ARGUMENT(bool, SupportsThumbMouseButtonNavigation)
 
+		/** Whether to show a throbber overlay during browser initialization. */
+		SLATE_ARGUMENT(bool, ShowInitialThrobber)
+
 		/** Opaque background color used before a document is loaded and when no document color is specified. */
 		SLATE_ARGUMENT(FColor, BackgroundColor)
 
@@ -345,4 +349,6 @@ private:
 	/** A delegate that is invoked when when the browser needs to dismiss all dialogs */
 	FSimpleDelegate OnDismissAllDialogs;
 
+	/** The initial throbber setting */
+	bool bShowInitialThrobber;
 };
