Serialization: Added support for TMap property serialization

diff --git a/Engine/Source/Runtime/Serialization/Private/Backends/JsonStructSerializerBackend.cpp b/Engine/Source/Runtime/Serialization/Private/Backends/JsonStructSerializerBackend.cpp
--- a/Engine/Source/Runtime/Serialization/Private/Backends/JsonStructSerializerBackend.cpp
+++ b/Engine/Source/Runtime/Serialization/Private/Backends/JsonStructSerializerBackend.cpp
@@ -11,28 +11,40 @@ namespace JsonStructSerializerBackend
 {
 	// Writes a property value to the serialization output.
 	template<typename ValueType>
-	void WritePropertyValue( const TSharedRef<TJsonWriter<UCS2CHAR>> JsonWriter, UProperty* Property, const ValueType& Value )
+	void WritePropertyValue(const TSharedRef<TJsonWriter<UCS2CHAR>> JsonWriter, const FStructSerializerState& State, const ValueType& Value)
 	{
-		if ((Property == nullptr) || (Property->ArrayDim > 1) || (Property->GetOuter()->GetClass() == UArrayProperty::StaticClass()))
+		if ((State.ValueProperty == nullptr) || (State.ValueProperty->ArrayDim > 1) || (State.ValueProperty->GetOuter()->GetClass() == UArrayProperty::StaticClass()))
 		{
 			JsonWriter->WriteValue(Value);
 		}
+		else if (State.KeyProperty != nullptr)
+		{
+			FString KeyString;
+			State.KeyProperty->ExportTextItem(KeyString, State.KeyData, nullptr, nullptr, PPF_None);
+			JsonWriter->WriteValue(KeyString, Value);
+		}
 		else
 		{
-			JsonWriter->WriteValue(Property->GetName(), Value);
+			JsonWriter->WriteValue(State.ValueProperty->GetName(), Value);
 		}
 	}
 
 	// Writes a null value to the serialization output.
-	void WriteNull( const TSharedRef<TJsonWriter<UCS2CHAR>> JsonWriter, UProperty* Property )
+	void WriteNull(const TSharedRef<TJsonWriter<UCS2CHAR>> JsonWriter, const FStructSerializerState& State)
 	{
-		if ((Property == nullptr) || (Property->ArrayDim > 1) || (Property->GetOuter()->GetClass() == UArrayProperty::StaticClass()))
+		if ((State.ValueProperty == nullptr) || (State.ValueProperty->ArrayDim > 1) || (State.ValueProperty->GetOuter()->GetClass() == UArrayProperty::StaticClass()))
 		{
 			JsonWriter->WriteNull();
 		}
+		else if (State.KeyProperty != nullptr)
+		{
+			FString KeyString;
+			State.KeyProperty->ExportTextItem(KeyString, State.KeyData, nullptr, nullptr, PPF_None);
+			JsonWriter->WriteNull(KeyString);
+		}
 		else
 		{
-			JsonWriter->WriteNull(Property->GetName());
+			JsonWriter->WriteNull(State.ValueProperty->GetName());
 		}
 	}
 }
@@ -49,6 +61,12 @@ void FJsonStructSerializerBackend::BeginArray(const FStructSerializerState& Stat
 	{
 		JsonWriter->WriteArrayStart();
 	}
+	else if (State.KeyProperty != nullptr)
+	{
+		FString KeyString;
+		State.KeyProperty->ExportTextItem(KeyString, State.KeyData, nullptr, nullptr, PPF_None);
+		JsonWriter->WriteArrayStart(KeyString);
+	}
 	else
 	{
 		JsonWriter->WriteArrayStart(State.ValueProperty->GetName());
@@ -56,6 +74,27 @@ void FJsonStructSerializerBackend::BeginArray(const FStructSerializerState& Stat
 }
 
 
+void FJsonStructSerializerBackend::BeginDictionary(const FStructSerializerState& State)
+{
+	UObject* Outer = State.ValueProperty->GetOuter();
+
+	if ((Outer != nullptr) && (Outer->GetClass() == UArrayProperty::StaticClass()))
+	{
+		JsonWriter->WriteObjectStart();
+	}
+	else if (State.KeyProperty != nullptr)
+	{
+		FString KeyString;
+		State.KeyProperty->ExportTextItem(KeyString, State.KeyData, nullptr, nullptr, PPF_None);
+		JsonWriter->WriteObjectStart(KeyString);
+	}
+	else
+	{
+		JsonWriter->WriteObjectStart(State.ValueProperty->GetName());
+	}
+}
+
+
 void FJsonStructSerializerBackend::BeginStructure(const FStructSerializerState& State)
 {
 	if (State.ValueProperty != nullptr)
@@ -66,6 +105,12 @@ void FJsonStructSerializerBackend::BeginStructure(const FStructSerializerState&
 		{
 			JsonWriter->WriteObjectStart();
 		}
+		else if (State.KeyProperty != nullptr)
+		{
+			FString KeyString;
+			State.KeyProperty->ExportTextItem(KeyString, State.KeyData, nullptr, nullptr, PPF_None);
+			JsonWriter->WriteObjectStart(KeyString);
+		}
 		else
 		{
 			JsonWriter->WriteObjectStart(State.ValueProperty->GetName());
@@ -84,6 +129,12 @@ void FJsonStructSerializerBackend::EndArray(const FStructSerializerState& /*Stat
 }
 
 
+void FJsonStructSerializerBackend::EndDictionary(const FStructSerializerState& /*State*/)
+{
+	JsonWriter->WriteObjectEnd();
+}
+
+
 void FJsonStructSerializerBackend::EndStructure(const FStructSerializerState& /*State*/)
 {
 	JsonWriter->WriteObjectEnd();
@@ -103,7 +154,7 @@ void FJsonStructSerializerBackend::WriteProperty(const FStructSerializerState& S
 	// booleans
 	if (State.ValueType == UBoolProperty::StaticClass())
 	{
-		WritePropertyValue(JsonWriter, State.ValueProperty, Cast<UBoolProperty>(State.ValueProperty)->GetPropertyValue_InContainer(State.ValueData, ArrayIndex));
+		WritePropertyValue(JsonWriter, State, Cast<UBoolProperty>(State.ValueProperty)->GetPropertyValue_InContainer(State.ValueData, ArrayIndex));
 	}
 
 	// unsigned bytes & enumerations
@@ -113,78 +164,78 @@ void FJsonStructSerializerBackend::WriteProperty(const FStructSerializerState& S
 
 		if (ByteProperty->IsEnum())
 		{
-			WritePropertyValue(JsonWriter, State.ValueProperty, ByteProperty->Enum->GetEnumName(ByteProperty->GetPropertyValue_InContainer(State.ValueData, ArrayIndex)));
+			WritePropertyValue(JsonWriter, State, ByteProperty->Enum->GetEnumName(ByteProperty->GetPropertyValue_InContainer(State.ValueData, ArrayIndex)));
 		}
 		else
 		{
-			WritePropertyValue(JsonWriter, State.ValueProperty, (double)Cast<UByteProperty>(State.ValueProperty)->GetPropertyValue_InContainer(State.ValueData, ArrayIndex));
+			WritePropertyValue(JsonWriter, State, (double)Cast<UByteProperty>(State.ValueProperty)->GetPropertyValue_InContainer(State.ValueData, ArrayIndex));
 		}			
 	}
 
 	// floating point numbers
 	else if (State.ValueType == UDoubleProperty::StaticClass())
 	{
-		WritePropertyValue(JsonWriter, State.ValueProperty, Cast<UDoubleProperty>(State.ValueProperty)->GetPropertyValue_InContainer(State.ValueData, ArrayIndex));
+		WritePropertyValue(JsonWriter, State, Cast<UDoubleProperty>(State.ValueProperty)->GetPropertyValue_InContainer(State.ValueData, ArrayIndex));
 	}
 	else if (State.ValueType == UFloatProperty::StaticClass())
 	{
-		WritePropertyValue(JsonWriter, State.ValueProperty, Cast<UFloatProperty>(State.ValueProperty)->GetPropertyValue_InContainer(State.ValueData, ArrayIndex));
+		WritePropertyValue(JsonWriter, State, Cast<UFloatProperty>(State.ValueProperty)->GetPropertyValue_InContainer(State.ValueData, ArrayIndex));
 	}
 
 	// signed integers
 	else if (State.ValueType == UIntProperty::StaticClass())
 	{
-		WritePropertyValue(JsonWriter, State.ValueProperty, (double)Cast<UIntProperty>(State.ValueProperty)->GetPropertyValue_InContainer(State.ValueData, ArrayIndex));
+		WritePropertyValue(JsonWriter, State, (double)Cast<UIntProperty>(State.ValueProperty)->GetPropertyValue_InContainer(State.ValueData, ArrayIndex));
 	}
 	else if (State.ValueType == UInt8Property::StaticClass())
 	{
-		WritePropertyValue(JsonWriter, State.ValueProperty, (double)Cast<UInt8Property>(State.ValueProperty)->GetPropertyValue_InContainer(State.ValueData, ArrayIndex));
+		WritePropertyValue(JsonWriter, State, (double)Cast<UInt8Property>(State.ValueProperty)->GetPropertyValue_InContainer(State.ValueData, ArrayIndex));
 	}
 	else if (State.ValueType == UInt16Property::StaticClass())
 	{
-		WritePropertyValue(JsonWriter, State.ValueProperty, (double)Cast<UInt16Property>(State.ValueProperty)->GetPropertyValue_InContainer(State.ValueData, ArrayIndex));
+		WritePropertyValue(JsonWriter, State, (double)Cast<UInt16Property>(State.ValueProperty)->GetPropertyValue_InContainer(State.ValueData, ArrayIndex));
 	}
 	else if (State.ValueType == UInt64Property::StaticClass())
 	{
-		WritePropertyValue(JsonWriter, State.ValueProperty, (double)Cast<UInt64Property>(State.ValueProperty)->GetPropertyValue_InContainer(State.ValueData, ArrayIndex));
+		WritePropertyValue(JsonWriter, State, (double)Cast<UInt64Property>(State.ValueProperty)->GetPropertyValue_InContainer(State.ValueData, ArrayIndex));
 	}
 
 	// unsigned integers
 	else if (State.ValueType == UUInt16Property::StaticClass())
 	{
-		WritePropertyValue(JsonWriter, State.ValueProperty, (double)Cast<UUInt16Property>(State.ValueProperty)->GetPropertyValue_InContainer(State.ValueData, ArrayIndex));
+		WritePropertyValue(JsonWriter, State, (double)Cast<UUInt16Property>(State.ValueProperty)->GetPropertyValue_InContainer(State.ValueData, ArrayIndex));
 	}
 	else if (State.ValueType == UUInt32Property::StaticClass())
 	{
-		WritePropertyValue(JsonWriter, State.ValueProperty, (double)Cast<UUInt32Property>(State.ValueProperty)->GetPropertyValue_InContainer(State.ValueData, ArrayIndex));
+		WritePropertyValue(JsonWriter, State, (double)Cast<UUInt32Property>(State.ValueProperty)->GetPropertyValue_InContainer(State.ValueData, ArrayIndex));
 	}
 	else if (State.ValueType == UUInt64Property::StaticClass())
 	{
-		WritePropertyValue(JsonWriter, State.ValueProperty, (double)Cast<UUInt64Property>(State.ValueProperty)->GetPropertyValue_InContainer(State.ValueData, ArrayIndex));
+		WritePropertyValue(JsonWriter, State, (double)Cast<UUInt64Property>(State.ValueProperty)->GetPropertyValue_InContainer(State.ValueData, ArrayIndex));
 	}
 
 	// names, strings & text
 	else if (State.ValueType == UNameProperty::StaticClass())
 	{
-		WritePropertyValue(JsonWriter, State.ValueProperty, Cast<UNameProperty>(State.ValueProperty)->GetPropertyValue_InContainer(State.ValueData, ArrayIndex).ToString());
+		WritePropertyValue(JsonWriter, State, Cast<UNameProperty>(State.ValueProperty)->GetPropertyValue_InContainer(State.ValueData, ArrayIndex).ToString());
 	}
 	else if (State.ValueType == UStrProperty::StaticClass())
 	{
-		WritePropertyValue(JsonWriter, State.ValueProperty, Cast<UStrProperty>(State.ValueProperty)->GetPropertyValue_InContainer(State.ValueData, ArrayIndex));
+		WritePropertyValue(JsonWriter, State, Cast<UStrProperty>(State.ValueProperty)->GetPropertyValue_InContainer(State.ValueData, ArrayIndex));
 	}
 	else if (State.ValueType == UTextProperty::StaticClass())
 	{
-		WritePropertyValue(JsonWriter, State.ValueProperty, Cast<UTextProperty>(State.ValueProperty)->GetPropertyValue_InContainer(State.ValueData, ArrayIndex).ToString());
+		WritePropertyValue(JsonWriter, State, Cast<UTextProperty>(State.ValueProperty)->GetPropertyValue_InContainer(State.ValueData, ArrayIndex).ToString());
 	}
 
 	// classes & objects
 	else if (State.ValueType == UClassProperty::StaticClass())
 	{
-		WritePropertyValue(JsonWriter, State.ValueProperty, Cast<UClassProperty>(State.ValueProperty)->GetPropertyValue_InContainer(State.ValueData, ArrayIndex)->GetPathName());
+		WritePropertyValue(JsonWriter, State, Cast<UClassProperty>(State.ValueProperty)->GetPropertyValue_InContainer(State.ValueData, ArrayIndex)->GetPathName());
 	}
 	else if (State.ValueType == UObjectProperty::StaticClass())
 	{
-		WriteNull(JsonWriter, State.ValueProperty);
+		WriteNull(JsonWriter, State);
 	}
 	
 	// unsupported property type
diff --git a/Engine/Source/Runtime/Serialization/Private/StructSerializer.cpp b/Engine/Source/Runtime/Serialization/Private/StructSerializer.cpp
--- a/Engine/Source/Runtime/Serialization/Private/StructSerializer.cpp
+++ b/Engine/Source/Runtime/Serialization/Private/StructSerializer.cpp
@@ -61,6 +61,8 @@ void FStructSerializer::Serialize( const void* Struct, UStruct& TypeInfo, IStruc
 		FStructSerializerState NewState;
 		{
 			NewState.HasBeenProcessed = false;
+			NewState.KeyData = nullptr;
+			NewState.KeyProperty = nullptr;
 			NewState.ValueData = Struct;
 			NewState.ValueProperty = nullptr;
 			NewState.ValueType = &TypeInfo;
@@ -79,7 +81,7 @@ void FStructSerializer::Serialize( const void* Struct, UStruct& TypeInfo, IStruc
 		{
 			if (!CurrentState.HasBeenProcessed)
 			{
-				const void* NewData = CurrentState.ValueData;
+				const void* ValueData = CurrentState.ValueData;
 
 				// write object start
 				if (CurrentState.ValueProperty != nullptr)
@@ -88,7 +90,7 @@ void FStructSerializer::Serialize( const void* Struct, UStruct& TypeInfo, IStruc
 
 					if ((Outer == nullptr) || (Outer->GetClass() != UArrayProperty::StaticClass()))
 					{
-						NewData = CurrentState.ValueProperty->ContainerPtrToValuePtr<void>(CurrentState.ValueData);
+						ValueData = CurrentState.ValueProperty->ContainerPtrToValuePtr<void>(CurrentState.ValueData);
 					}
 				}
 
@@ -130,7 +132,9 @@ void FStructSerializer::Serialize( const void* Struct, UStruct& TypeInfo, IStruc
 					FStructSerializerState NewState;
 					{
 						NewState.HasBeenProcessed = false;
-						NewState.ValueData = NewData;
+						NewState.KeyData = nullptr;
+						NewState.KeyProperty = nullptr;
+						NewState.ValueData = ValueData;
 						NewState.ValueProperty = *It;
 						NewState.ValueType = It->GetClass();
 					}
@@ -162,7 +166,7 @@ void FStructSerializer::Serialize( const void* Struct, UStruct& TypeInfo, IStruc
 
 				UArrayProperty* ArrayProperty = Cast<UArrayProperty>(CurrentState.ValueProperty);
 				FScriptArrayHelper ArrayHelper(ArrayProperty, ArrayProperty->ContainerPtrToValuePtr<void>(CurrentState.ValueData));
-				UProperty* Inner = ArrayProperty->Inner;
+				UProperty* ValueProperty = ArrayProperty->Inner;
 
 				// push elements on stack (in reverse order)
 				for (int Index = ArrayHelper.Num() - 1; Index >= 0; --Index)
@@ -170,9 +174,11 @@ void FStructSerializer::Serialize( const void* Struct, UStruct& TypeInfo, IStruc
 					FStructSerializerState NewState;
 					{
 						NewState.HasBeenProcessed = false;
+						NewState.KeyData = nullptr;
+						NewState.KeyProperty = nullptr;
 						NewState.ValueData = ArrayHelper.GetRawPtr(Index);
-						NewState.ValueProperty = Inner;
-						NewState.ValueType = Inner->GetClass();
+						NewState.ValueProperty = ValueProperty;
+						NewState.ValueType = ValueProperty->GetClass();
 					}
 
 					StateStack.Push(NewState);
@@ -184,6 +190,44 @@ void FStructSerializer::Serialize( const void* Struct, UStruct& TypeInfo, IStruc
 			}
 		}
 
+		// maps
+		else if (CurrentState.ValueType == UMapProperty::StaticClass())
+		{
+			if (!CurrentState.HasBeenProcessed)
+			{
+				Backend.BeginDictionary(CurrentState);
+
+				CurrentState.HasBeenProcessed = true;
+				StateStack.Push(CurrentState);
+
+				UMapProperty* MapProperty = Cast<UMapProperty>(CurrentState.ValueProperty);
+				FScriptMapHelper MapHelper(MapProperty, MapProperty->ContainerPtrToValuePtr<void>(CurrentState.ValueData));
+				UProperty* ValueProperty = MapProperty->ValueProp;
+				
+				// push key-value pairs on stack (in reverse order)
+				for (int Index = MapHelper.Num() - 1; Index >= 0; --Index)
+				{
+					uint8* PairPtr = MapHelper.GetPairPtr(Index);;
+
+					FStructSerializerState NewState;
+					{
+						NewState.HasBeenProcessed = false;
+						NewState.KeyData = PairPtr + MapProperty->MapLayout.KeyOffset;
+						NewState.KeyProperty = MapProperty->KeyProp;
+						NewState.ValueData = PairPtr;// + MapProperty->MapLayout.ValueOffset;
+						NewState.ValueProperty = ValueProperty;
+						NewState.ValueType = ValueProperty->GetClass();
+					}
+
+					StateStack.Push(NewState);
+				}
+			}
+			else
+			{
+				Backend.EndDictionary(CurrentState);
+			}
+		}
+
 		// static arrays
 		else if (CurrentState.ValueProperty->ArrayDim > 1)
 		{
diff --git a/Engine/Source/Runtime/Serialization/Private/Tests/StructSerializerTest.cpp b/Engine/Source/Runtime/Serialization/Private/Tests/StructSerializerTest.cpp
--- a/Engine/Source/Runtime/Serialization/Private/Tests/StructSerializerTest.cpp
+++ b/Engine/Source/Runtime/Serialization/Private/Tests/StructSerializerTest.cpp
@@ -31,7 +31,7 @@ namespace StructSerializerTest
 			Test.TestTrue(TEXT("Deserialization must succeed"), FStructDeserializer::Deserialize(TestStruct2, DeserializerBackend, Policies));
 		}
 
-		// test for equality
+		// test numerics
 		Test.TestEqual<int8>(TEXT("Numerics.Int8 value must be the same before and after de-/serialization"), TestStruct.Numerics.Int8, TestStruct2.Numerics.Int8);
 		Test.TestEqual<int16>(TEXT("Numerics.Int16 value must be the same before and after de-/serialization"), TestStruct.Numerics.Int16, TestStruct2.Numerics.Int16);
 		Test.TestEqual<int32>(TEXT("Numerics.Int32 value must be the same before and after de-/serialization"), TestStruct.Numerics.Int32, TestStruct2.Numerics.Int32);
@@ -43,13 +43,16 @@ namespace StructSerializerTest
 		Test.TestEqual<float>(TEXT("Numerics.Float value must be the same before and after de-/serialization"), TestStruct.Numerics.Float, TestStruct2.Numerics.Float);
 		Test.TestEqual<double>(TEXT("Numerics.Double value must be the same before and after de-/serialization"), TestStruct.Numerics.Double, TestStruct2.Numerics.Double);
 
+		// test booleans
 		Test.TestEqual<bool>(TEXT("Booleans.BoolFalse must be the same before and after de-/serialization"), TestStruct.Booleans.BoolFalse, TestStruct2.Booleans.BoolFalse);
 		Test.TestEqual<bool>(TEXT("Booleans.BoolTrue must be the same before and after de-/serialization"), TestStruct.Booleans.BoolTrue, TestStruct2.Booleans.BoolTrue);
 		Test.TestEqual<uint32>(TEXT("Booleans.Bitfield must be the same before and after de-/serialization"), TestStruct.Booleans.Bitfield, TestStruct2.Booleans.Bitfield);
 
+		// test objects
 		Test.TestEqual<TSubclassOf<class UObject>>(TEXT("Objects.Class must be the same before and after de-/serialization"), TestStruct.Objects.Class, TestStruct2.Objects.Class);
 		Test.TestEqual<UObject*>(TEXT("Objects.ObjectPtr must be the same before and after de-/serialization"), TestStruct.Objects.ObjectPtr, TestStruct2.Objects.ObjectPtr);
 
+		// test built-ins
 		Test.TestEqual<FGuid>(TEXT("Builtins.Guid must be the same before and after de-/serialization"), TestStruct.Builtins.Guid, TestStruct2.Builtins.Guid);
 		Test.TestEqual<FName>(TEXT("Builtins.Name must be the same before and after de-/serialization"), TestStruct.Builtins.Name, TestStruct2.Builtins.Name);
 		Test.TestEqual<FString>(TEXT("Builtins.String must be the same before and after de-/serialization"), TestStruct.Builtins.String, TestStruct2.Builtins.String);
@@ -57,6 +60,7 @@ namespace StructSerializerTest
 		Test.TestEqual<FString>(TEXT("Builtins.Text must be the same before and after de-/serialization"), TestStruct.Builtins.Text.ToString(), TestStruct2.Builtins.Text.ToString());
 		Test.TestEqual<FVector>(TEXT("Builtins.Vector must be the same before and after de-/serialization"), TestStruct.Builtins.Vector, TestStruct2.Builtins.Vector);
 
+		// test arrays
 		Test.TestEqual<TArray<int32>>(TEXT("Arrays.Int32Array must be the same before and after de-/serialization"), TestStruct.Arrays.Int32Array, TestStruct2.Arrays.Int32Array);
 		Test.TestEqual<int32>(TEXT("Arrays.StaticSingleElement[0] must be the same before and after de-/serialization"), TestStruct.Arrays.StaticSingleElement[0], TestStruct2.Arrays.StaticSingleElement[0]);
 		Test.TestEqual<int32>(TEXT("Arrays.StaticInt32Array[0] must be the same before and after de-/serialization"), TestStruct.Arrays.StaticInt32Array[0], TestStruct2.Arrays.StaticInt32Array[0]);
@@ -66,6 +70,10 @@ namespace StructSerializerTest
 		Test.TestEqual<float>(TEXT("Arrays.StaticFloatArray[1] must be the same before and after de-/serialization"), TestStruct.Arrays.StaticFloatArray[1], TestStruct2.Arrays.StaticFloatArray[1]);
 		Test.TestEqual<float>(TEXT("Arrays.StaticFloatArray[2] must be the same before and after de-/serialization"), TestStruct.Arrays.StaticFloatArray[2], TestStruct2.Arrays.StaticFloatArray[2]);
 		Test.TestEqual<TArray<FVector>>(TEXT("Arrays.VectorArray must be the same before and after de-/serialization"), TestStruct.Arrays.VectorArray, TestStruct2.Arrays.VectorArray);
+
+		// test maps
+//		Test.TestTrue(TEXT("Maps.IntToString must be the same before and after de-/serialization"), TestStruct.Maps.IntToString.OrderIndependentCompareEqual(TestStruct2.Maps.IntToString));
+//		Test.TestTrue(TEXT("Maps.StringToVector must be the same before and after de-/serialization"), TestStruct.Maps.StringToVector.OrderIndependentCompareEqual(TestStruct2.Maps.StringToVector));
 	}
 }
 
@@ -90,8 +98,7 @@ bool FJsonStructSerializerTest::RunTest( const FString& Parameters )
 		StructSerializerTest::TestSerialization(*this, SerializerBackend, DeserializerBackend);
 
 		// uncomment this to look at the serialized data
-		//uint8* Data = Buffer.GetData();
-		//GLog->Logf(TEXT("%s"), (TCHAR*)Data);
+//		GLog->Logf(TEXT("%s"), (TCHAR*)Buffer.GetData());
 	}
 
 	return true;
diff --git a/Engine/Source/Runtime/Serialization/Private/Tests/StructSerializerTestTypes.h b/Engine/Source/Runtime/Serialization/Private/Tests/StructSerializerTestTypes.h
--- a/Engine/Source/Runtime/Serialization/Private/Tests/StructSerializerTestTypes.h
+++ b/Engine/Source/Runtime/Serialization/Private/Tests/StructSerializerTestTypes.h
@@ -203,6 +203,37 @@ struct FStructSerializerArrayTestStruct
 };
 
 
+/**
+ * Test structure for map properties.
+ */
+USTRUCT()
+struct FStructSerializerMapTestStruct
+{
+	GENERATED_USTRUCT_BODY()
+
+	UPROPERTY()
+	TMap<int32, FString> IntToString;
+
+	UPROPERTY()
+	TMap<FString, FVector> StringToVector;
+
+	/** Default constructor. */
+	FStructSerializerMapTestStruct()
+	{
+		IntToString.Add(1, TEXT("One"));
+		IntToString.Add(2, TEXT("Two"));
+		IntToString.Add(3, TEXT("Three"));
+
+		StringToVector.Add(TEXT("Vector000"), FVector(0.0f, 0.0f, 0.0f));
+		StringToVector.Add(TEXT("Vector123"), FVector(1.0f, 2.0f, 3.0f));
+		StringToVector.Add(TEXT("Vector666"), FVector(6.0f, 6.0f, 6.0f));
+	}
+
+	/** Creates an uninitialized instance. */
+	FStructSerializerMapTestStruct( ENoInit ) { }
+};
+
+
 /**
  * Test structure for all supported types.
  */
@@ -226,6 +257,9 @@ struct FStructSerializerTestStruct
 	UPROPERTY()
 	FStructSerializerArrayTestStruct Arrays;
 
+	UPROPERTY()
+	FStructSerializerMapTestStruct Maps;
+
 	/** Default constructor. */
 	FStructSerializerTestStruct() { }
 
@@ -236,5 +270,6 @@ struct FStructSerializerTestStruct
 		, Objects(NoInit)
 		, Builtins(NoInit)
 		, Arrays(NoInit)
+		, Maps(NoInit)
 	{ }
 };
diff --git a/Engine/Source/Runtime/Serialization/Public/Backends/JsonStructSerializerBackend.h b/Engine/Source/Runtime/Serialization/Public/Backends/JsonStructSerializerBackend.h
--- a/Engine/Source/Runtime/Serialization/Public/Backends/JsonStructSerializerBackend.h
+++ b/Engine/Source/Runtime/Serialization/Public/Backends/JsonStructSerializerBackend.h
@@ -38,8 +38,10 @@ public:
 	// IStructSerializerBackend interface
 
 	virtual void BeginArray(const FStructSerializerState& State) override;
+	virtual void BeginDictionary(const FStructSerializerState& State) override;
 	virtual void BeginStructure(const FStructSerializerState& State) override;
 	virtual void EndArray(const FStructSerializerState& State) override;
+	virtual void EndDictionary(const FStructSerializerState& State) override;
 	virtual void EndStructure(const FStructSerializerState& State) override;
 	virtual void WriteComment(const FString& Comment) override;
 	virtual void WriteProperty(const FStructSerializerState& State, int32 ArrayIndex = 0) override;
diff --git a/Engine/Source/Runtime/Serialization/Public/IStructSerializerBackend.h b/Engine/Source/Runtime/Serialization/Public/IStructSerializerBackend.h
--- a/Engine/Source/Runtime/Serialization/Public/IStructSerializerBackend.h
+++ b/Engine/Source/Runtime/Serialization/Public/IStructSerializerBackend.h
@@ -16,6 +16,12 @@ struct FStructSerializerState
 	/** Holds a flag indicating whether the property has been processed. */
 	bool HasBeenProcessed;
 
+	/** Holds a pointer to the key property's data. */
+	const void* KeyData;
+
+	/** Holds the key property's meta data (only used for TMap). */
+	UProperty* KeyProperty;
+
 	/** Holds a pointer to the property value's data. */
 	const void* ValueData;
 
@@ -44,6 +50,13 @@ public:
 	 */
 	virtual void BeginArray(const FStructSerializerState& State) = 0;
 
+	/**
+	 * Signals the beginning of a dictionary.
+	 *
+	 * @param Property The property that holds the dictionary.
+	 */
+	virtual void BeginDictionary(const FStructSerializerState& State) = 0;
+
 	/**
 	 * Signals the beginning of a child structure.
 	 *
@@ -64,6 +77,13 @@ public:
 	 */
 	virtual void EndArray(const FStructSerializerState& State) = 0;
 
+	/**
+	 * Signals the end of a dictionary.
+	 *
+	 * @param Property The property that holds the dictionary.
+	 */
+	virtual void EndDictionary(const FStructSerializerState& State) = 0;
+
 	/**
 	 * Signals the end of an object.
 	 *
@@ -87,6 +107,8 @@ public:
 	 *
 	 * Depending on the context, properties to be written can be either object properties or array elements.
 	 *
+	 * State.KeyProperty points to the key property that holds the data to write.
+	 * State.KeyData points to the key property's data.
 	 * State.ValueProperty points to the property that holds the value to write.
 	 * State.ValueData points to the actual data to write.
 	 * State.TypeInfo contains the data's type information
diff --git a/Engine/Source/Runtime/WebBrowser/Private/WebJSStructSerializerBackend.cpp b/Engine/Source/Runtime/WebBrowser/Private/WebJSStructSerializerBackend.cpp
--- a/Engine/Source/Runtime/WebBrowser/Private/WebJSStructSerializerBackend.cpp
+++ b/Engine/Source/Runtime/WebBrowser/Private/WebJSStructSerializerBackend.cpp
@@ -3,17 +3,18 @@
 #include "WebBrowserPrivatePCH.h"
 #include "WebJSStructSerializerBackend.h"
 
+
 #if WITH_CEF3
 
 /* Private methods
  *****************************************************************************/
 
-void FWebJSStructSerializerBackend::AddNull(UProperty* Property)
+void FWebJSStructSerializerBackend::AddNull(const FStructSerializerState& State)
 {
 	StackItem& Current = Stack.Top();
 	switch (Current.Kind) {
 		case StackItem::STYPE_DICTIONARY:
-			Current.DictionaryValue->SetNull(*Property->GetName());
+			Current.DictionaryValue->SetNull(*State.ValueProperty->GetName());
 			break;
 		case StackItem::STYPE_LIST:
 			Current.ListValue->SetNull(Current.ListValue->GetSize());
@@ -21,12 +22,13 @@ void FWebJSStructSerializerBackend::AddNull(UProperty* Property)
 	}
 }
 
-void FWebJSStructSerializerBackend::Add(UProperty* Property, bool Value)
+
+void FWebJSStructSerializerBackend::Add(const FStructSerializerState& State, bool Value)
 {
 	StackItem& Current = Stack.Top();
 	switch (Current.Kind) {
 		case StackItem::STYPE_DICTIONARY:
-			Current.DictionaryValue->SetBool(*Property->GetName(), Value);
+			Current.DictionaryValue->SetBool(*State.ValueProperty->GetName(), Value);
 			break;
 		case StackItem::STYPE_LIST:
 			Current.ListValue->SetBool(Current.ListValue->GetSize(), Value);
@@ -34,12 +36,13 @@ void FWebJSStructSerializerBackend::Add(UProperty* Property, bool Value)
 	}
 }
 
-void FWebJSStructSerializerBackend::Add(UProperty* Property, int32 Value)
+
+void FWebJSStructSerializerBackend::Add(const FStructSerializerState& State, int32 Value)
 {
 	StackItem& Current = Stack.Top();
 	switch (Current.Kind) {
 		case StackItem::STYPE_DICTIONARY:
-			Current.DictionaryValue->SetInt(*Property->GetName(), Value);
+			Current.DictionaryValue->SetInt(*State.ValueProperty->GetName(), Value);
 			break;
 		case StackItem::STYPE_LIST:
 			Current.ListValue->SetInt(Current.ListValue->GetSize(), Value);
@@ -47,12 +50,13 @@ void FWebJSStructSerializerBackend::Add(UProperty* Property, int32 Value)
 	}
 }
 
-void FWebJSStructSerializerBackend::Add(UProperty* Property, double Value)
+
+void FWebJSStructSerializerBackend::Add(const FStructSerializerState& State, double Value)
 {
 	StackItem& Current = Stack.Top();
 	switch (Current.Kind) {
 		case StackItem::STYPE_DICTIONARY:
-			Current.DictionaryValue->SetDouble(*Property->GetName(), Value);
+			Current.DictionaryValue->SetDouble(*State.ValueProperty->GetName(), Value);
 			break;
 		case StackItem::STYPE_LIST:
 			Current.ListValue->SetDouble(Current.ListValue->GetSize(), Value);
@@ -60,12 +64,13 @@ void FWebJSStructSerializerBackend::Add(UProperty* Property, double Value)
 	}
 }
 
-void FWebJSStructSerializerBackend::Add(UProperty* Property, FString Value)
+
+void FWebJSStructSerializerBackend::Add(const FStructSerializerState& State, FString Value)
 {
 	StackItem& Current = Stack.Top();
 	switch (Current.Kind) {
 		case StackItem::STYPE_DICTIONARY:
-			Current.DictionaryValue->SetString(*Property->GetName(), *Value);
+			Current.DictionaryValue->SetString(*State.ValueProperty->GetName(), *Value);
 			break;
 		case StackItem::STYPE_LIST:
 			Current.ListValue->SetString(Current.ListValue->GetSize(), *Value);
@@ -73,12 +78,13 @@ void FWebJSStructSerializerBackend::Add(UProperty* Property, FString Value)
 	}
 }
 
-void FWebJSStructSerializerBackend::Add(UProperty* Property, UObject* Value)
+
+void FWebJSStructSerializerBackend::Add(const FStructSerializerState& State, UObject* Value)
 {
 	StackItem& Current = Stack.Top();
 	switch (Current.Kind) {
 		case StackItem::STYPE_DICTIONARY:
-			Current.DictionaryValue->SetDictionary(*Property->GetName(), Scripting->ConvertObject(Value));
+			Current.DictionaryValue->SetDictionary(*State.ValueProperty->GetName(), Scripting->ConvertObject(Value));
 			break;
 		case StackItem::STYPE_LIST:
 			Current.ListValue->SetDictionary(Current.ListValue->GetSize(), Scripting->ConvertObject(Value));
@@ -86,6 +92,7 @@ void FWebJSStructSerializerBackend::Add(UProperty* Property, UObject* Value)
 	}
 }
 
+
 /* IStructSerializerBackend interface
  *****************************************************************************/
 
@@ -95,9 +102,25 @@ void FWebJSStructSerializerBackend::BeginArray(const FStructSerializerState& Sta
 	Stack.Push(StackItem(State.ValueProperty->GetName(), ListValue));
 }
 
+
+void FWebJSStructSerializerBackend::BeginDictionary(const FStructSerializerState& State)
+{
+	CefRefPtr<CefDictionaryValue> DictionaryValue = CefDictionaryValue::Create();
+	Stack.Push(StackItem(State.ValueProperty->GetName(), DictionaryValue));
+}
+
+
 void FWebJSStructSerializerBackend::BeginStructure(const FStructSerializerState& State)
 {
-	if (State.ValueProperty != nullptr)
+	if (State.KeyProperty != nullptr)
+	{
+		FString KeyString;
+		State.KeyProperty->ExportTextItem(KeyString, State.KeyData, nullptr, nullptr, PPF_None);
+
+		CefRefPtr<CefDictionaryValue> DictionaryValue = CefDictionaryValue::Create();
+		Stack.Push(StackItem(KeyString, DictionaryValue));
+	}
+	else if (State.ValueProperty != nullptr)
 	{
 		CefRefPtr<CefDictionaryValue> DictionaryValue = CefDictionaryValue::Create();
 		Stack.Push(StackItem(State.ValueProperty->GetName(), DictionaryValue));
@@ -109,6 +132,7 @@ void FWebJSStructSerializerBackend::BeginStructure(const FStructSerializerState&
 	}
 }
 
+
 void FWebJSStructSerializerBackend::EndArray(const FStructSerializerState& /*State*/)
 {
 	StackItem Previous = Stack.Pop();
@@ -127,6 +151,12 @@ void FWebJSStructSerializerBackend::EndArray(const FStructSerializerState& /*Sta
 }
 
 
+void FWebJSStructSerializerBackend::EndDictionary(const FStructSerializerState& State)
+{
+	EndStructure(State);
+}
+
+
 void FWebJSStructSerializerBackend::EndStructure(const FStructSerializerState& /*State*/)
 {
 	StackItem Previous = Stack.Pop();
@@ -160,11 +190,10 @@ void FWebJSStructSerializerBackend::WriteComment(const FString& Comment)
 
 void FWebJSStructSerializerBackend::WriteProperty(const FStructSerializerState& State, int32 ArrayIndex)
 {
-
 	// booleans
 	if (State.ValueType == UBoolProperty::StaticClass())
 	{
-		Add(State.ValueProperty, Cast<UBoolProperty>(State.ValueProperty)->GetPropertyValue_InContainer(State.ValueData, ArrayIndex));
+		Add(State, Cast<UBoolProperty>(State.ValueProperty)->GetPropertyValue_InContainer(State.ValueData, ArrayIndex));
 	}
 
 	// unsigned bytes & enumerations
@@ -174,78 +203,78 @@ void FWebJSStructSerializerBackend::WriteProperty(const FStructSerializerState&
 
 		if (ByteProperty->IsEnum())
 		{
-			Add(State.ValueProperty, ByteProperty->Enum->GetEnumName(ByteProperty->GetPropertyValue_InContainer(State.ValueData, ArrayIndex)));
+			Add(State, ByteProperty->Enum->GetEnumName(ByteProperty->GetPropertyValue_InContainer(State.ValueData, ArrayIndex)));
 		}
 		else
 		{
-			Add(State.ValueProperty, (double)Cast<UByteProperty>(State.ValueProperty)->GetPropertyValue_InContainer(State.ValueData, ArrayIndex));
+			Add(State, (double)Cast<UByteProperty>(State.ValueProperty)->GetPropertyValue_InContainer(State.ValueData, ArrayIndex));
 		}			
 	}
 
 	// floating point numbers
 	else if (State.ValueType == UDoubleProperty::StaticClass())
 	{
-		Add(State.ValueProperty, Cast<UDoubleProperty>(State.ValueProperty)->GetPropertyValue_InContainer(State.ValueData, ArrayIndex));
+		Add(State, Cast<UDoubleProperty>(State.ValueProperty)->GetPropertyValue_InContainer(State.ValueData, ArrayIndex));
 	}
 	else if (State.ValueType == UFloatProperty::StaticClass())
 	{
-		Add(State.ValueProperty, Cast<UFloatProperty>(State.ValueProperty)->GetPropertyValue_InContainer(State.ValueData, ArrayIndex));
+		Add(State, Cast<UFloatProperty>(State.ValueProperty)->GetPropertyValue_InContainer(State.ValueData, ArrayIndex));
 	}
 
 	// signed integers
 	else if (State.ValueType == UIntProperty::StaticClass())
 	{
-		Add(State.ValueProperty, (int32)Cast<UIntProperty>(State.ValueProperty)->GetPropertyValue_InContainer(State.ValueData, ArrayIndex));
+		Add(State, (int32)Cast<UIntProperty>(State.ValueProperty)->GetPropertyValue_InContainer(State.ValueData, ArrayIndex));
 	}
 	else if (State.ValueType == UInt8Property::StaticClass())
 	{
-		Add(State.ValueProperty, (int32)Cast<UInt8Property>(State.ValueProperty)->GetPropertyValue_InContainer(State.ValueData, ArrayIndex));
+		Add(State, (int32)Cast<UInt8Property>(State.ValueProperty)->GetPropertyValue_InContainer(State.ValueData, ArrayIndex));
 	}
 	else if (State.ValueType == UInt16Property::StaticClass())
 	{
-		Add(State.ValueProperty, (int32)Cast<UInt16Property>(State.ValueProperty)->GetPropertyValue_InContainer(State.ValueData, ArrayIndex));
+		Add(State, (int32)Cast<UInt16Property>(State.ValueProperty)->GetPropertyValue_InContainer(State.ValueData, ArrayIndex));
 	}
 	else if (State.ValueType == UInt64Property::StaticClass())
 	{
-		Add(State.ValueProperty, (double)Cast<UInt64Property>(State.ValueProperty)->GetPropertyValue_InContainer(State.ValueData, ArrayIndex));
+		Add(State, (double)Cast<UInt64Property>(State.ValueProperty)->GetPropertyValue_InContainer(State.ValueData, ArrayIndex));
 	}
 
 	// unsigned integers
 	else if (State.ValueType == UUInt16Property::StaticClass())
 	{
-		Add(State.ValueProperty, (int32)Cast<UUInt16Property>(State.ValueProperty)->GetPropertyValue_InContainer(State.ValueData, ArrayIndex));
+		Add(State, (int32)Cast<UUInt16Property>(State.ValueProperty)->GetPropertyValue_InContainer(State.ValueData, ArrayIndex));
 	}
 	else if (State.ValueType == UUInt32Property::StaticClass())
 	{
-		Add(State.ValueProperty, (double)Cast<UUInt32Property>(State.ValueProperty)->GetPropertyValue_InContainer(State.ValueData, ArrayIndex));
+		Add(State, (double)Cast<UUInt32Property>(State.ValueProperty)->GetPropertyValue_InContainer(State.ValueData, ArrayIndex));
 	}
 	else if (State.ValueType == UUInt64Property::StaticClass())
 	{
-		Add(State.ValueProperty, (double)Cast<UUInt64Property>(State.ValueProperty)->GetPropertyValue_InContainer(State.ValueData, ArrayIndex));
+		Add(State, (double)Cast<UUInt64Property>(State.ValueProperty)->GetPropertyValue_InContainer(State.ValueData, ArrayIndex));
 	}
 
 	// names & strings
 	else if (State.ValueType == UNameProperty::StaticClass())
 	{
-		Add(State.ValueProperty, Cast<UNameProperty>(State.ValueProperty)->GetPropertyValue_InContainer(State.ValueData, ArrayIndex).ToString());
+		Add(State, Cast<UNameProperty>(State.ValueProperty)->GetPropertyValue_InContainer(State.ValueData, ArrayIndex).ToString());
 	}
 	else if (State.ValueType == UStrProperty::StaticClass())
 	{
-		Add(State.ValueProperty, Cast<UStrProperty>(State.ValueProperty)->GetPropertyValue_InContainer(State.ValueData, ArrayIndex));
+		Add(State, Cast<UStrProperty>(State.ValueProperty)->GetPropertyValue_InContainer(State.ValueData, ArrayIndex));
 	}
 	else if (State.ValueType == UTextProperty::StaticClass())
 	{
-		Add(State.ValueProperty, Cast<UTextProperty>(State.ValueProperty)->GetPropertyValue_InContainer(State.ValueData, ArrayIndex).ToString());
+		Add(State, Cast<UTextProperty>(State.ValueProperty)->GetPropertyValue_InContainer(State.ValueData, ArrayIndex).ToString());
 	}
 
 	// classes & objects
 	else if (State.ValueType == UClassProperty::StaticClass())
 	{
-		Add(State.ValueProperty, Cast<UClassProperty>(State.ValueProperty)->GetPropertyValue_InContainer(State.ValueData, ArrayIndex)->GetPathName());
+		Add(State, Cast<UClassProperty>(State.ValueProperty)->GetPropertyValue_InContainer(State.ValueData, ArrayIndex)->GetPathName());
 	}
 	else if (State.ValueType == UObjectProperty::StaticClass())
 	{
-		Add(State.ValueProperty, Cast<UObjectProperty>(State.ValueProperty)->GetPropertyValue_InContainer(State.ValueData, ArrayIndex));
+		Add(State, Cast<UObjectProperty>(State.ValueProperty)->GetPropertyValue_InContainer(State.ValueData, ArrayIndex));
 	}
 	
 	// unsupported property type
@@ -255,4 +284,5 @@ void FWebJSStructSerializerBackend::WriteProperty(const FStructSerializerState&
 	}
 }
 
+
 #endif
diff --git a/Engine/Source/Runtime/WebBrowser/Private/WebJSStructSerializerBackend.h b/Engine/Source/Runtime/WebBrowser/Private/WebJSStructSerializerBackend.h
--- a/Engine/Source/Runtime/WebBrowser/Private/WebJSStructSerializerBackend.h
+++ b/Engine/Source/Runtime/WebBrowser/Private/WebJSStructSerializerBackend.h
@@ -5,6 +5,8 @@
 #include "IStructSerializerBackend.h"
 #include "Core.h"
 #include "WebJSScripting.h"
+
+
 #if WITH_CEF3
 
 // forward declarations
@@ -12,16 +14,19 @@ class UProperty;
 class UStruct;
 
 #if PLATFORM_WINDOWS
-#include "AllowWindowsPlatformTypes.h"
+	#include "AllowWindowsPlatformTypes.h"
 #endif
+
 #pragma push_macro("OVERRIDE")
 #undef OVERRIDE // cef headers provide their own OVERRIDE macro
 #include "include/cef_values.h"
 #pragma pop_macro("OVERRIDE")
+
 #if PLATFORM_WINDOWS
-#include "HideWindowsPlatformTypes.h"
+	#include "HideWindowsPlatformTypes.h"
 #endif
 
+
 /**
  * Implements a writer for UStruct serialization using CefDictionary.
  */
@@ -35,8 +40,7 @@ public:
 	 *	 */
 	FWebJSStructSerializerBackend(TSharedPtr<FWebJSScripting> InScripting)
 		: Scripting(InScripting), Stack(), Result()
-	{
-	}
+	{ }
 
 	CefRefPtr<CefDictionaryValue> GetResult()
 	{
@@ -48,8 +52,10 @@ public:
 	// IStructSerializerBackend interface
 
 	virtual void BeginArray(const FStructSerializerState& State) override;
+	virtual void BeginDictionary(const FStructSerializerState& State) override;
 	virtual void BeginStructure(const FStructSerializerState& State) override;
 	virtual void EndArray(const FStructSerializerState& State) override;
+	virtual void EndDictionary(const FStructSerializerState& State) override;
 	virtual void EndStructure(const FStructSerializerState& State) override;
 	virtual void WriteComment(const FString& Comment) override;
 	virtual void WriteProperty(const FStructSerializerState& State, int32 ArrayIndex = 0) override;
@@ -82,12 +88,14 @@ private:
 	TSharedPtr<FWebJSScripting> Scripting;
 	TArray<StackItem> Stack;
 	CefRefPtr<CefDictionaryValue> Result;
-	void AddNull(UProperty* Property);
-	void Add(UProperty* Property, bool Value);
-	void Add(UProperty* Property, int32 Value);
-	void Add(UProperty* Property, double Value);
-	void Add(UProperty* Property, FString Value);
-	void Add(UProperty* Property, UObject* Value);
 
+	void AddNull(const FStructSerializerState& State);
+	void Add(const FStructSerializerState& State, bool Value);
+	void Add(const FStructSerializerState& State, int32 Value);
+	void Add(const FStructSerializerState& State, double Value);
+	void Add(const FStructSerializerState& State, FString Value);
+	void Add(const FStructSerializerState& State, UObject* Value);
 };
-#endif
\ No newline at end of file
+
+
+#endif
