Merging CL# 2694350, 2694409 using UE4-To-UE4-LauncherDev
The web browser ShowErrorMessage argument is now hooked up to drive the showing of error messages.

diff --git a/Engine/Source/Runtime/WebBrowser/Private/WebBrowserHandler.cpp b/Engine/Source/Runtime/WebBrowser/Private/WebBrowserHandler.cpp
--- a/Engine/Source/Runtime/WebBrowser/Private/WebBrowserHandler.cpp
+++ b/Engine/Source/Runtime/WebBrowser/Private/WebBrowserHandler.cpp
@@ -11,7 +11,6 @@
 
 #if WITH_CEF3
 FWebBrowserHandler::FWebBrowserHandler()
-	: ShowErrorMessage(true)
 {
 }
 
@@ -174,23 +173,23 @@ void FWebBrowserHandler::OnLoadError(CefRefPtr<CefBrowser> Browser,
 
 		if (BrowserWindow.IsValid())
 		{
-			BrowserWindow->NotifyDocumentError();
-		}
-	}
+			// Display a load error message. Note: The user's code will still have a chance to handle this error after this error message is displayed.
+			if (BrowserWindow->IsShowingErrorMessages())
+			{
+				FFormatNamedArguments Args;
+				{
+					Args.Add(TEXT("FailedUrl"), FText::FromString(FailedUrl.ToWString().c_str()));
+					Args.Add(TEXT("ErrorText"), FText::FromString(ErrorText.ToWString().c_str()));
+					Args.Add(TEXT("ErrorCode"), FText::AsNumber(InErrorCode));
+				}
+				FText ErrorMsg = FText::Format(LOCTEXT("WebBrowserLoadError", "Failed to load URL {FailedUrl} with error {ErrorText} ({ErrorCode})."), Args);
+				FString ErrorHTML = TEXT("<html><body bgcolor=\"white\"><h2>") + ErrorMsg.ToString() + TEXT("</h2></body></html>");
+
+				Frame->LoadString(*ErrorHTML, FailedUrl);
+			}
 
-	// Display a load error message.
-	if (ShowErrorMessage)
-	{
-		FFormatNamedArguments Args;
-		{
-			Args.Add(TEXT("FailedUrl"), FText::FromString(FailedUrl.ToWString().c_str()));
-			Args.Add(TEXT("ErrorText"), FText::FromString(ErrorText.ToWString().c_str()));
-			Args.Add(TEXT("ErrorCode"), FText::AsNumber(InErrorCode));
+			BrowserWindow->NotifyDocumentError();
 		}
-		FText ErrorMsg = FText::Format(LOCTEXT("WebBrowserLoadError", "Failed to load URL {FailedUrl} with error {ErrorText} ({ErrorCode})."), Args);
-		FString ErrorHTML = TEXT("<html><body bgcolor=\"white\"><h2>") + ErrorMsg.ToString() + TEXT("</h2></body></html>");
-
-		Frame->LoadString(*ErrorHTML, FailedUrl);
 	}
 }
 
diff --git a/Engine/Source/Runtime/WebBrowser/Private/WebBrowserHandler.h b/Engine/Source/Runtime/WebBrowser/Private/WebBrowserHandler.h
--- a/Engine/Source/Runtime/WebBrowser/Private/WebBrowserHandler.h
+++ b/Engine/Source/Runtime/WebBrowser/Private/WebBrowserHandler.h
@@ -66,12 +66,6 @@ public:
 		BrowserPopupFeatures = InPopupFeatures;
 	}
 
-	/** Sets whether to show messages on loading errors. */
-	void SetShowErrorMessage(bool InShowErrorMessage)
-	{
-		ShowErrorMessage = InShowErrorMessage;
-	}
-
 
 public:
 
@@ -246,9 +240,6 @@ private:
 	/** Stores popup window features and settings */
 	TSharedPtr<FWebBrowserPopupFeatures> BrowserPopupFeatures;
 
-	/** Whether to show an error message in case of loading errors. */
-	bool ShowErrorMessage;
-
 	// Include the default reference counting implementation.
 	IMPLEMENT_REFCOUNTING(FWebBrowserHandler);
 };
