WebBrowser: Add a method to FWebBrowserSingleton to delete cookies.
Add a flag for enabling Chrome Devtools keyboard shortcut (CMD/CTRL-SHIFT-I) in shipping builds as well as debug builds.
Update 64bit version of UnrealCEFSubprocess.exe missing from previous CEF upgrade (dummy merge)

Merging CL#2675892, CL#2704418, and CL#2711307 using UE4-To-UE4-LauncherDev

diff --git a/Engine/Source/Runtime/WebBrowser/Private/WebBrowserHandler.cpp b/Engine/Source/Runtime/WebBrowser/Private/WebBrowserHandler.cpp
--- a/Engine/Source/Runtime/WebBrowser/Private/WebBrowserHandler.cpp
+++ b/Engine/Source/Runtime/WebBrowser/Private/WebBrowserHandler.cpp
@@ -2,7 +2,9 @@
 
 #include "WebBrowserPrivatePCH.h"
 #include "WebBrowserHandler.h"
+#include "WebBrowserModule.h"
 #include "WebBrowserWindow.h"
+#include "IWebBrowserSingleton.h"
 #include "WebBrowserSingleton.h"
 #include "WebBrowserPopupFeatures.h"
 #include "SlateApplication.h"
@@ -11,8 +13,7 @@
 
 #if WITH_CEF3
 FWebBrowserHandler::FWebBrowserHandler()
-{
-}
+{ }
 
 void FWebBrowserHandler::OnTitleChange(CefRefPtr<CefBrowser> Browser, const CefString& Title)
 {
@@ -396,20 +397,19 @@ bool FWebBrowserHandler::OnKeyEvent(CefRefPtr<CefBrowser> Browser,
 	const CefKeyEvent& Event,
 	CefEventHandle OsEvent)
 {
-#if UE_BUILD_DEBUG
-	// Show dev tools on CMD/CTRL+ALT+I
+	// Show dev tools on CMD/CTRL+SHIFT+I
 	if( (Event.type == KEYEVENT_RAWKEYDOWN || Event.type == KEYEVENT_KEYDOWN) &&
 #if PLATFORM_MAC
 		(Event.modifiers == (EVENTFLAG_COMMAND_DOWN | EVENTFLAG_SHIFT_DOWN)) &&
 #else
 		(Event.modifiers == (EVENTFLAG_CONTROL_DOWN | EVENTFLAG_SHIFT_DOWN)) &&
 #endif
-		(Event.unmodified_character == 'i' || Event.unmodified_character == 'I')
+		(Event.unmodified_character == 'i' || Event.unmodified_character == 'I') &&
+		IWebBrowserModule::Get().GetSingleton()->IsDevToolsShortcutEnabled()
 	  )
 	{
 		return ShowDevTools(Browser);
 	}
-#endif
 
 #if PLATFORM_MAC
 	// We need to handle standard Copy/Paste/etc... shortcuts on OS X
diff --git a/Engine/Source/Runtime/WebBrowser/Private/WebBrowserSingleton.cpp b/Engine/Source/Runtime/WebBrowser/Private/WebBrowserSingleton.cpp
--- a/Engine/Source/Runtime/WebBrowser/Private/WebBrowserSingleton.cpp
+++ b/Engine/Source/Runtime/WebBrowser/Private/WebBrowserSingleton.cpp
@@ -113,6 +113,7 @@ namespace {
 }
 
 FWebBrowserSingleton::FWebBrowserSingleton()
+	: bDevToolsShortcutEnabled(UE_BUILD_DEBUG)
 {
 #if WITH_CEF3
 	// The FWebBrowserSingleton must be initialized on the game thread
@@ -338,6 +339,60 @@ FString FWebBrowserSingleton::GetCurrentLocaleCode()
 	return LocaleCode;
 }
 
+#if WITH_CEF3
+namespace
+{
+	// Task for executing a callback on a given thread.
+	class FDeleteCookiesNotificationTask
+		: public CefTask
+	{
+		TFunction<void (int)> Callback;
+		int Value;
+	public:
+		FDeleteCookiesNotificationTask(const TFunction<void (int)>& InCallback, int InValue)
+			: Callback(InCallback)
+			, Value(InValue)
+		{}
+
+		virtual void Execute() override
+		{
+			Callback(Value);
+		}
+
+		IMPLEMENT_REFCOUNTING(FDeleteCookiesNotificationTask);
+
+	};
+
+	// Callback that will invoke the callback with the result on the UI thread.
+	class FDeleteCookiesFunctionCallback
+		: public CefDeleteCookiesCallback
+	{
+		TFunction<void (int)> Callback;
+	public:
+		FDeleteCookiesFunctionCallback(const TFunction<void (int)>& InCallback)
+			: Callback(InCallback)
+		{}
+
+		virtual void OnComplete(int NumDeleted) override
+		{
+			// We're on the IO thread, so we'll have to schedule the callback on the main thread
+			CefPostTask(TID_UI, new FDeleteCookiesNotificationTask(Callback, NumDeleted));
+		}
+
+		IMPLEMENT_REFCOUNTING(FDeleteCookiesFunctionCallback);
+	};
+}
+#endif
+
+void FWebBrowserSingleton::DeleteBrowserCookies(FString URL, FString CookieName, TFunction<void (int)> Completed)
+{
+#if WITH_CEF3
+	CefRefPtr<FDeleteCookiesFunctionCallback> Callback = Completed ? new FDeleteCookiesFunctionCallback(Completed) : nullptr;
+	CefCookieManager::GetGlobalManager(nullptr)->DeleteCookies(*URL, *CookieName, Callback);
+#endif
+}
+
+
 // Cleanup macros to avoid having them leak outside this source file
 #undef CEF3_BIN_DIR
 #undef CEF3_FRAMEWORK_DIR
diff --git a/Engine/Source/Runtime/WebBrowser/Private/WebBrowserSingleton.h b/Engine/Source/Runtime/WebBrowser/Private/WebBrowserSingleton.h
--- a/Engine/Source/Runtime/WebBrowser/Private/WebBrowserSingleton.h
+++ b/Engine/Source/Runtime/WebBrowser/Private/WebBrowserSingleton.h
@@ -64,6 +64,17 @@ public:
 		bool ShowErrorMessage = true,
 		FColor BackgroundColor = FColor(255, 255, 255, 255)) override;
 
+	virtual void DeleteBrowserCookies(FString URL = TEXT(""), FString CookieName = TEXT(""), TFunction<void (int)> Completed = nullptr) override;
+
+	virtual bool IsDevToolsShortcutEnabled() override
+	{
+		return bDevToolsShortcutEnabled;
+	}
+
+	virtual void SetDevToolsShortcutEnabled(bool Value) override
+	{
+		bDevToolsShortcutEnabled = Value;
+	}
 
 public:
 
@@ -80,4 +91,5 @@ private:
 	/** List of currently existing browser windows */
 	TArray<TWeakPtr<FWebBrowserWindow>>	WindowInterfaces;
 #endif
+	bool bDevToolsShortcutEnabled;
 };
diff --git a/Engine/Source/Runtime/WebBrowser/Public/IWebBrowserSingleton.h b/Engine/Source/Runtime/WebBrowser/Public/IWebBrowserSingleton.h
--- a/Engine/Source/Runtime/WebBrowser/Public/IWebBrowserSingleton.h
+++ b/Engine/Source/Runtime/WebBrowser/Public/IWebBrowserSingleton.h
@@ -50,4 +50,37 @@ public:
 		TOptional<FString> ContentsToLoad = TOptional<FString>(), 
 		bool ShowErrorMessage = true,
 		FColor BackgroundColor = FColor(255, 255, 255, 255)) = 0;
+
+	/**
+	 * Delete all browser cookies.
+	 *
+	 * Removes all matching cookies. Leave both URL and CookieName blank to delete the entire cookie database.
+	 * The actual deletion will be scheduled on the browser IO thread, so the operation may not have completed when the function returns.
+	 *
+	 * @param URL The base URL to match when searching for cookies to remove. Use blank to match all URLs.
+	 * @param CookieName The name of the cookie to delete. If left unspecified, all cookies will be removed.
+	 * @param Completed A callback function that will be invoked asynchronously on the game thread when the deletion is complete passing in the number of deleted objects.
+	 */
+	virtual void DeleteBrowserCookies(FString URL = TEXT(""), FString CookieName = TEXT(""), TFunction<void (int)> Completed = nullptr) = 0;
+
+	/**
+	 * Enable or disable CTRL/CMD-SHIFT-I shortcut to show the Chromium Dev tools window.
+	 * The value defaults to true on debug builds, otherwise false.
+	 *
+	 * The relevant handlers for spawning new browser windows have to be set up correctly in addition to this flag being true before anything is shown.
+	 *
+	 * @param Value a boolean value to enable or disable the keyboard shortcut.
+	 */
+	virtual void SetDevToolsShortcutEnabled(bool Value) = 0;
+
+
+	/**
+	 * Returns wether the CTRL/CMD-SHIFT-I shortcut to show the Chromium Dev tools window is enabled.
+	 *
+	 * The relevant handlers for spawning new browser windows have to be set up correctly in addition to this flag being true before anything is shown.
+	 *
+	 * @return a boolean value indicating whether the keyboard shortcut is enabled or not.
+	 */
+	virtual bool IsDevToolsShortcutEnabled() = 0;
+
 };
