Created new browser adapter that overrides the default history behavior so it routes through the portal.
Merging using UE4-To-UE4-LauncherDev

diff --git a/Engine/Source/Runtime/WebBrowser/Private/SWebBrowser.cpp b/Engine/Source/Runtime/WebBrowser/Private/SWebBrowser.cpp
--- a/Engine/Source/Runtime/WebBrowser/Private/SWebBrowser.cpp
+++ b/Engine/Source/Runtime/WebBrowser/Private/SWebBrowser.cpp
@@ -7,6 +7,7 @@
 #include "IWebBrowserSingleton.h"
 #include "IWebBrowserWindow.h"
 #include "WebBrowserViewport.h"
+#include "IWebBrowserAdapter.h"
 
 #define LOCTEXT_NAMESPACE "WebBrowser"
 
@@ -384,7 +385,17 @@ void SWebBrowser::HandleBrowserWindowDocumentStateChanged(EWebBrowserDocumentSta
 	switch (NewState)
 	{
 	case EWebBrowserDocumentState::Completed:
-		OnLoadCompleted.ExecuteIfBound();
+		{
+			if (BrowserWindow.IsValid())
+			{
+				for (auto Adapter : Adapters)
+				{
+					Adapter->ConnectTo(BrowserWindow.ToSharedRef());
+				}
+			}
+
+			OnLoadCompleted.ExecuteIfBound();
+		}
 		break;
 
 	case EWebBrowserDocumentState::Error:
@@ -517,6 +528,24 @@ void SWebBrowser::UnbindUObject(const FString& Name, UObject* Object, bool bIsPe
 	}
 }
 
+void SWebBrowser::BindAdapter(const TSharedRef<IWebBrowserAdapter>& Adapter)
+{
+	Adapters.Add(Adapter);
+	if (BrowserWindow.IsValid())
+	{
+		Adapter->ConnectTo(BrowserWindow.ToSharedRef());
+	}
+}
+
+void SWebBrowser::UnbindAdapter(const TSharedRef<IWebBrowserAdapter>& Adapter)
+{
+	Adapters.Remove(Adapter);
+	if (BrowserWindow.IsValid())
+	{
+		Adapter->DisconnectFrom(BrowserWindow.ToSharedRef());
+	}
+}
+
 void SWebBrowser::HandleShowPopup(const FIntRect& PopupSize)
 {
 	check(!PopupMenuPtr.IsValid())
diff --git a/Engine/Source/Runtime/WebBrowser/Private/WebBrowserAdapter.cpp b/Engine/Source/Runtime/WebBrowser/Private/WebBrowserAdapter.cpp
new file mode 100644
--- /dev/null
+++ b/Engine/Source/Runtime/WebBrowser/Private/WebBrowserAdapter.cpp
@@ -0,0 +1,97 @@
+// Copyright 1998-2015 Epic Games, Inc. All Rights Reserved.
+
+#include "WebBrowserPrivatePCH.h"
+#include "IWebBrowserAdapter.h"
+#include "IWebBrowserWindow.h"
+#include "GCObject.h"
+#include "UObjectGlobals.h"
+
+class FDefaultWebBrowserAdapter
+	: public IWebBrowserAdapter
+	, public FGCObject
+
+{
+public:
+
+	virtual FString GetName() const override
+	{
+		return Name;
+	}
+
+	virtual bool IsPermanent() const override
+	{
+		return bIsPermanent;
+	}
+
+	virtual void ConnectTo(const TSharedRef<IWebBrowserWindow>& BrowserWindow) override
+	{
+		if (JSBridge != nullptr)
+		{
+			BrowserWindow->BindUObject(Name, JSBridge, bIsPermanent);
+		}
+
+		if (!ConnectScriptText.IsEmpty())
+		{
+			BrowserWindow->ExecuteJavascript(ConnectScriptText);
+		}
+	}
+
+	virtual void DisconnectFrom(const TSharedRef<IWebBrowserWindow>& BrowserWindow) override
+	{
+		if (!DisconnectScriptText.IsEmpty())
+		{
+			BrowserWindow->ExecuteJavascript(DisconnectScriptText);
+		}
+
+		if (JSBridge != nullptr)
+		{
+			BrowserWindow->UnbindUObject(Name, JSBridge, bIsPermanent);
+		}
+	}
+
+	// FGCObject API
+	virtual void AddReferencedObjects(FReferenceCollector& Collector) override
+	{
+		if (JSBridge != nullptr)
+		{
+			Collector.AddReferencedObject(JSBridge);
+		}
+	}
+
+private:
+
+	FDefaultWebBrowserAdapter(
+		const FString InName,
+		const FString InConnectScriptText,
+		const FString InDisconnectScriptText,
+		UObject* InJSBridge,
+		const bool InIsPermanent)
+		: Name(InName)
+		, ConnectScriptText(InConnectScriptText)
+		, DisconnectScriptText(InDisconnectScriptText)
+		, JSBridge(InJSBridge)
+		, bIsPermanent(InIsPermanent)
+	{ }
+
+private:
+
+	const FString Name;
+	const FString ConnectScriptText;
+	const FString DisconnectScriptText;
+
+	UObject* JSBridge;
+
+	const bool bIsPermanent;
+
+	friend FWebBrowserAdapterFactory;
+};
+
+TSharedRef<IWebBrowserAdapter> FWebBrowserAdapterFactory::Create(const FString& Name, UObject* JSBridge, bool IsPermanent)
+{
+	return MakeShareable(new FDefaultWebBrowserAdapter(Name, FString(), FString(), JSBridge, IsPermanent));
+}
+
+TSharedRef<IWebBrowserAdapter> FWebBrowserAdapterFactory::Create(const FString& Name, UObject* JSBridge, bool IsPermanent, const FString& ConnectScriptText, const FString& DisconnectScriptText)
+{
+	return MakeShareable(new FDefaultWebBrowserAdapter(Name, ConnectScriptText, DisconnectScriptText, JSBridge, IsPermanent));
+}
\ No newline at end of file
diff --git a/Engine/Source/Runtime/WebBrowser/Public/IWebBrowserAdapter.h b/Engine/Source/Runtime/WebBrowser/Public/IWebBrowserAdapter.h
new file mode 100644
--- /dev/null
+++ b/Engine/Source/Runtime/WebBrowser/Public/IWebBrowserAdapter.h
@@ -0,0 +1,27 @@
+// Copyright 1998-2015 Epic Games, Inc. All Rights Reserved.
+#pragma once
+
+class IWebBrowserWindow;
+
+class IWebBrowserAdapter
+{
+public:
+
+	virtual FString GetName() const = 0;
+
+	virtual bool IsPermanent() const = 0;
+
+	virtual void ConnectTo(const TSharedRef<IWebBrowserWindow>& BrowserWindow) = 0;
+
+	virtual void DisconnectFrom(const TSharedRef<IWebBrowserWindow>& BrowserWindow) = 0;
+
+};
+
+class WEBBROWSER_API FWebBrowserAdapterFactory 
+{ 
+public: 
+
+	static TSharedRef<IWebBrowserAdapter> Create(const FString& Name, UObject* JSBridge, bool IsPermanent); 
+
+	static TSharedRef<IWebBrowserAdapter> Create(const FString& Name, UObject* JSBridge, bool IsPermanent, const FString& ConnectScriptText, const FString& DisconnectScriptText);
+}; 
\ No newline at end of file
diff --git a/Engine/Source/Runtime/WebBrowser/Public/SWebBrowser.h b/Engine/Source/Runtime/WebBrowser/Public/SWebBrowser.h
--- a/Engine/Source/Runtime/WebBrowser/Public/SWebBrowser.h
+++ b/Engine/Source/Runtime/WebBrowser/Public/SWebBrowser.h
@@ -10,6 +10,7 @@ class IWebBrowserWindow;
 class FWebBrowserViewport;
 class UObject;
 class IWebBrowserPopupFeatures;
+class IWebBrowserAdapter;
 
 DECLARE_DELEGATE_RetVal_TwoParams(bool, FOnBeforePopupDelegate, FString, FString);
 DECLARE_DELEGATE_RetVal_TwoParams(bool, FOnCreateWindowDelegate, const TWeakPtr<IWebBrowserWindow>&, const TWeakPtr<IWebBrowserPopupFeatures>&);
@@ -190,6 +191,10 @@ public:
 	 */
 	void UnbindUObject(const FString& Name, UObject* Object, bool bIsPermanent = true);
 
+	void BindAdapter(const TSharedRef<IWebBrowserAdapter>& Adapter);
+
+	void UnbindAdapter(const TSharedRef<IWebBrowserAdapter>& Adapter);
+
 private:
 
 	/** Returns true if the browser can navigate backwards. */
@@ -298,6 +303,8 @@ private:
 	/** The actual viewport widget. Required to update its tool tip property. */
 	TSharedPtr<SViewport> ViewportWidget;
 
+	TArray<TSharedRef<IWebBrowserAdapter>> Adapters;
+
 	/**
 	 * An interface pointer to a menu object presenting a popup.
 	 * Pointer is null when a popup is not visible.
