Fixing non-unity issue
Merging using UE4-To-UE4-LauncherDev

diff --git a/Engine/Source/Runtime/WebBrowser/Private/WebBrowserAdapter.cpp b/Engine/Source/Runtime/WebBrowser/Private/WebBrowserAdapter.cpp
--- a/Engine/Source/Runtime/WebBrowser/Private/WebBrowserAdapter.cpp
+++ b/Engine/Source/Runtime/WebBrowser/Private/WebBrowserAdapter.cpp
@@ -3,8 +3,8 @@
 #include "WebBrowserPrivatePCH.h"
 #include "IWebBrowserAdapter.h"
 #include "IWebBrowserWindow.h"
+#include "CoreUObject.h"
 #include "GCObject.h"
-#include "UObjectGlobals.h"
 
 class FDefaultWebBrowserAdapter
 	: public IWebBrowserAdapter
