First implementation of SWebBrowser for Android

diff --git a/Engine/Build/Android/Java/proguard-project.txt b/Engine/Build/Android/Java/proguard-project.txt
--- a/Engine/Build/Android/Java/proguard-project.txt
+++ b/Engine/Build/Android/Java/proguard-project.txt
@@ -25,6 +25,9 @@
 -keep class com.epicgames.ue4.MessageBox01 {
    public *;
 }
+-keep class com.epicgames.ue4.WebViewControl {
+   public *;
+}
 
 # Needed...
 -keep class com.epicgames.ue4.GameActivity$InputDeviceInfo {
diff --git a/Engine/Build/Android/Java/src/com/epicgames/ue4/GameActivity.java b/Engine/Build/Android/Java/src/com/epicgames/ue4/GameActivity.java
--- a/Engine/Build/Android/Java/src/com/epicgames/ue4/GameActivity.java
+++ b/Engine/Build/Android/Java/src/com/epicgames/ue4/GameActivity.java
@@ -133,9 +133,11 @@ public class GameActivity extends NativeActivity
 	private AdView adView;
 	private boolean adInit = false;
 	private LinearLayout adLayout;
-	private LinearLayout activityLayout;
 	private int adGravity = Gravity.TOP;
 
+	// layout required by popups, e.g ads, native controls
+	LinearLayout activityLayout;
+
 	/** true when the application has requested that an ad be displayed */
 	private boolean adWantsToBeShown = false;
 
@@ -270,7 +272,13 @@ public class GameActivity extends NativeActivity
 		}
 
 		_activity = this;
-		
+
+		// layout required by popups, e.g ads, native controls
+		MarginLayoutParams params = new MarginLayoutParams(LayoutParams.WRAP_CONTENT, LayoutParams.WRAP_CONTENT);
+		params.setMargins(0, 0, 0, 0);
+		activityLayout = new LinearLayout(_activity);
+		_activity.setContentView(activityLayout, params);
+
 /*
 		// Turn on and unlock screen.. Assumption is that this
 		// will only really have an effect when for debug launching
@@ -960,7 +968,6 @@ public class GameActivity extends NativeActivity
 					adPopupWindow.setClippingEnabled(false);
 
 					adLayout = new LinearLayout(_activity);
-					activityLayout = new LinearLayout(_activity);
 
 					final int padding = (int)(-5*scale);
 					adLayout.setPadding(padding,padding,padding,padding);
@@ -973,8 +980,6 @@ public class GameActivity extends NativeActivity
 					adLayout.addView(adView, params);
 					adPopupWindow.setContentView(adLayout);
 
-					_activity.setContentView(activityLayout, params);
-
 					// set up our ad callbacks
 					_activity.adView.setAdListener(new AdListener()
 					{
diff --git a/Engine/Build/Android/Java/src/com/epicgames/ue4/WebViewControl.java b/Engine/Build/Android/Java/src/com/epicgames/ue4/WebViewControl.java
new file mode 100644
--- /dev/null
+++ b/Engine/Build/Android/Java/src/com/epicgames/ue4/WebViewControl.java
@@ -0,0 +1,85 @@
+// Copyright 1998-2015 Epic Games, Inc. All Rights Reserved.
+
+package com.epicgames.ue4;
+
+import android.content.res.Resources;
+import android.graphics.drawable.Drawable;
+import android.util.DisplayMetrics;
+import android.view.Gravity;
+import android.view.ViewGroup.LayoutParams;
+import android.view.ViewGroup.MarginLayoutParams;
+import android.widget.LinearLayout;
+import android.widget.PopupWindow;
+
+import android.webkit.WebView;
+import android.webkit.WebViewClient;
+
+class WebViewControl
+{
+	public WebViewControl()
+	{
+		GameActivity._activity.runOnUiThread(new Runnable()
+		{
+			@Override
+			public void run()
+			{
+				webView = new WebView(GameActivity._activity);
+				webView.getSettings().setJavaScriptEnabled(true);
+				webView.loadUrl("about:blank");
+			}
+		});
+	}
+
+	public void LoadURL(final String url)
+	{
+		GameActivity._activity.runOnUiThread(new Runnable()
+		{
+			@Override
+			public void run()
+			{
+				webView.loadUrl(url);
+			}
+		});
+	}
+
+	public void Update(final int x, final int y, final int width, final int height)
+	{
+		GameActivity._activity.runOnUiThread(new Runnable()
+		{
+			@Override
+			public void run()
+			{
+				if (webPopup == null)
+				{
+					webPopup = new PopupWindow(GameActivity._activity);
+					webPopup.setWidth(width);
+					webPopup.setHeight(height);
+					webPopup.setClippingEnabled(false);
+					webPopup.setOutsideTouchable(false); // non-modal
+					webPopup.setBackgroundDrawable((Drawable)null); // no border
+
+					webLayout = new LinearLayout(GameActivity._activity);
+					webLayout.setOrientation(LinearLayout.VERTICAL);
+					webLayout.setPadding(0, 0, 0, 0);
+
+					MarginLayoutParams params = new MarginLayoutParams(LayoutParams.MATCH_PARENT, LayoutParams.MATCH_PARENT);
+					params.setMargins(0, 0, 0, 0);
+					webLayout.addView(webView, params);
+
+					webPopup.setContentView(webLayout);
+					webPopup.showAtLocation(GameActivity._activity.activityLayout, Gravity.NO_GRAVITY, x, y);
+					webPopup.update();
+				}
+				else
+				{
+					webPopup.update(x, y, width, height);
+				}
+			}
+		});
+	}
+
+	// Web Views
+	private PopupWindow webPopup;
+	private LinearLayout webLayout;
+	private WebView webView;
+}
diff --git a/Engine/Plugins/Runtime/WebBrowserWidget/WebBrowserWidget.uplugin b/Engine/Plugins/Runtime/WebBrowserWidget/WebBrowserWidget.uplugin
--- a/Engine/Plugins/Runtime/WebBrowserWidget/WebBrowserWidget.uplugin
+++ b/Engine/Plugins/Runtime/WebBrowserWidget/WebBrowserWidget.uplugin
@@ -17,7 +17,8 @@
 			"WhitelistPlatforms": [
 				"Win64",
 				"Win32",
-				"Mac"
+				"Mac",
+				"Android"
 			]
 		}
 	],
diff --git a/Engine/Source/Runtime/WebBrowser/Private/Android/AndroidPlatformWebBrowser.cpp b/Engine/Source/Runtime/WebBrowser/Private/Android/AndroidPlatformWebBrowser.cpp
new file mode 100644
--- /dev/null
+++ b/Engine/Source/Runtime/WebBrowser/Private/Android/AndroidPlatformWebBrowser.cpp
@@ -0,0 +1,247 @@
+// Copyright 1998-2015 Epic Games, Inc. All Rights Reserved.
+
+#include "WebBrowserPrivatePCH.h"
+
+#if PLATFORM_ANDROID
+#include "AndroidPlatformWebBrowser.h"
+#include "AndroidWindow.h"
+#include "AndroidJava.h"
+
+class SAndroidWebBrowserWidget : public SLeafWidget
+{
+	SLATE_BEGIN_ARGS(SAndroidWebBrowserWidget)
+		: _InitialURL("about:blank")
+	{ }
+
+		SLATE_ARGUMENT(FString, InitialURL);
+	SLATE_END_ARGS()
+
+	void Construct(const FArguments& Args);
+
+	virtual int32 OnPaint(const FPaintArgs& Args, const FGeometry& AllottedGeometry, const FSlateRect& MyClippingRect, FSlateWindowElementList& OutDrawElements, int32 LayerId, const FWidgetStyle& InWidgetStyle, bool bParentEnabled) const override;
+
+	virtual FVector2D ComputeDesiredSize(float) const override;
+
+	void LoadURL(const FString& NewURL);
+
+protected:
+	// mutable to allow calling JWebView_Update from inside OnPaint (which is const)
+	mutable TOptional<FJavaClassObject> JWebView;
+	TOptional<FJavaClassMethod> JWebView_Update;
+	TOptional<FJavaClassMethod> JWebView_LoadURL;
+};
+
+void SAndroidWebBrowserWidget::Construct(const FArguments& Args)
+{
+	JWebView.Emplace("com/epicgames/ue4/WebViewControl", "()V");
+
+	JWebView_Update = JWebView->GetClassMethod("Update", "(IIII)V");
+	JWebView_LoadURL = JWebView->GetClassMethod("LoadURL", "(Ljava/lang/String;)V");
+
+	JWebView->CallMethod<void>(JWebView_LoadURL.GetValue(), FJavaClassObject::GetJString(Args._InitialURL));
+}
+
+int32 SAndroidWebBrowserWidget::OnPaint(const FPaintArgs& Args, const FGeometry& AllottedGeometry, const FSlateRect& MyClippingRect, FSlateWindowElementList& OutDrawElements, int32 LayerId, const FWidgetStyle& InWidgetStyle, bool bParentEnabled) const
+{
+	// Calculate UIScale, which can vary frame-to-frame thanks to device rotation
+	// UI Scale is calculated relative to vertical axis of 1280x720 / 720x1280
+	float UIScale;
+	FPlatformRect ScreenRect = FAndroidWindow::GetScreenRect();
+	int32_t ScreenWidth, ScreenHeight;
+	FAndroidWindow::CalculateSurfaceSize(FPlatformMisc::GetHardwareWindow(), ScreenWidth, ScreenHeight);
+	if (ScreenWidth > ScreenHeight)
+	{
+		UIScale = (float)ScreenHeight / (ScreenRect.Bottom - ScreenRect.Top);
+	}
+	else
+	{
+		UIScale = (float)ScreenHeight / (ScreenRect.Bottom - ScreenRect.Top);
+	}
+
+	FVector2D Position = AllottedGeometry.GetAccumulatedRenderTransform().GetTranslation() * UIScale;
+	FVector2D Size = TransformVector(AllottedGeometry.GetAccumulatedRenderTransform(), AllottedGeometry.GetLocalSize()) * UIScale;
+
+	// Convert position to integer coordinates
+	FIntPoint IntPos(FMath::RoundToInt(Position.X), FMath::RoundToInt(Position.Y));
+	// Convert size to integer taking the rounding of position into account to avoid double round-down or double round-up causing a noticeable error.
+	FIntPoint IntSize = FIntPoint(FMath::RoundToInt(Position.X + Size.X), FMath::RoundToInt(Size.Y + Position.Y)) - IntPos;
+
+	JWebView->CallMethod<void>(JWebView_Update.GetValue(), IntPos.X, IntPos.Y, IntSize.X, IntSize.Y);
+
+	return LayerId;
+}
+
+FVector2D SAndroidWebBrowserWidget::ComputeDesiredSize(float LayoutScaleMultiplier) const
+{
+	return FVector2D(640, 480);
+}
+
+void SAndroidWebBrowserWidget::LoadURL(const FString& NewURL)
+{
+	JWebView->CallMethod<void>(JWebView_LoadURL.GetValue(), FJavaClassObject::GetJString(NewURL));
+}
+
+FWebBrowserWindow::FWebBrowserWindow(FString InUrl, TOptional<FString> InContentsToLoad, bool InShowErrorMessage, bool InThumbMouseButtonNavigation, bool InUseTransparency)
+	: CurrentUrl(MoveTemp(InUrl))
+	, ContentsToLoad(MoveTemp(InContentsToLoad))
+{
+}
+
+FWebBrowserWindow::~FWebBrowserWindow()
+{
+}
+
+void FWebBrowserWindow::LoadURL(FString NewURL)
+{
+	BrowserWidget->LoadURL(NewURL);
+}
+
+void FWebBrowserWindow::LoadString(FString Contents, FString DummyURL)
+{
+}
+
+TSharedRef<SWidget> FWebBrowserWindow::CreateWidget(TAttribute<FVector2D> InViewportSize)
+{
+	TSharedRef<SAndroidWebBrowserWidget> BrowserWidgetRef =
+		SNew(SAndroidWebBrowserWidget)
+		.InitialURL(CurrentUrl);
+
+	BrowserWidget = BrowserWidgetRef;
+	return BrowserWidgetRef;
+}
+
+void FWebBrowserWindow::SetViewportSize(FIntPoint WindowSize)
+{
+}
+
+FSlateShaderResource* FWebBrowserWindow::GetTexture(bool bIsPopup /*= false*/)
+{
+	return nullptr;
+}
+
+bool FWebBrowserWindow::IsValid() const
+{
+	return false;
+}
+
+bool FWebBrowserWindow::IsInitialized() const
+{
+	return true;
+}
+
+bool FWebBrowserWindow::IsClosing() const
+{
+	return false;
+}
+
+EWebBrowserDocumentState FWebBrowserWindow::GetDocumentLoadingState() const
+{
+	return EWebBrowserDocumentState::Loading;
+}
+
+FString FWebBrowserWindow::GetTitle() const
+{
+	return "";
+}
+
+FString FWebBrowserWindow::GetUrl() const
+{
+	return "";
+}
+
+bool FWebBrowserWindow::OnKeyDown(const FKeyEvent& InKeyEvent)
+{
+	return false;
+}
+
+bool FWebBrowserWindow::OnKeyUp(const FKeyEvent& InKeyEvent)
+{
+	return false;
+}
+
+bool FWebBrowserWindow::OnKeyChar(const FCharacterEvent& InCharacterEvent)
+{
+	return false;
+}
+
+FReply FWebBrowserWindow::OnMouseButtonDown(const FGeometry& MyGeometry, const FPointerEvent& MouseEvent, bool bIsPopup)
+{
+	return FReply::Unhandled();
+}
+
+FReply FWebBrowserWindow::OnMouseButtonUp(const FGeometry& MyGeometry, const FPointerEvent& MouseEvent, bool bIsPopup)
+{
+	return FReply::Unhandled();
+}
+
+FReply FWebBrowserWindow::OnMouseButtonDoubleClick(const FGeometry& MyGeometry, const FPointerEvent& MouseEvent, bool bIsPopup)
+{
+	return FReply::Unhandled();
+}
+
+FReply FWebBrowserWindow::OnMouseMove(const FGeometry& MyGeometry, const FPointerEvent& MouseEvent, bool bIsPopup)
+{
+	return FReply::Unhandled();
+}
+
+FReply FWebBrowserWindow::OnMouseWheel(const FGeometry& MyGeometry, const FPointerEvent& MouseEvent, bool bIsPopup)
+{
+	return FReply::Unhandled();
+}
+
+void FWebBrowserWindow::OnFocus(bool SetFocus, bool bIsPopup)
+{
+}
+
+void FWebBrowserWindow::OnCaptureLost()
+{
+}
+
+bool FWebBrowserWindow::CanGoBack() const
+{
+	return false;
+}
+
+void FWebBrowserWindow::GoBack()
+{
+}
+
+bool FWebBrowserWindow::CanGoForward() const
+{
+	return false;
+}
+
+void FWebBrowserWindow::GoForward()
+{
+}
+
+bool FWebBrowserWindow::IsLoading() const
+{
+	return false;
+}
+
+void FWebBrowserWindow::Reload()
+{
+}
+
+void FWebBrowserWindow::StopLoad()
+{
+}
+
+void FWebBrowserWindow::ExecuteJavascript(const FString& Script)
+{
+}
+
+void FWebBrowserWindow::CloseBrowser(bool bForce)
+{
+}
+
+void FWebBrowserWindow::BindUObject(const FString& Name, UObject* Object, bool bIsPermanent /*= true*/)
+{
+}
+
+void FWebBrowserWindow::UnbindUObject(const FString& Name, UObject* Object /*= nullptr*/, bool bIsPermanent /*= true*/)
+{
+}
+
+#endif
diff --git a/Engine/Source/Runtime/WebBrowser/Private/Android/AndroidPlatformWebBrowser.h b/Engine/Source/Runtime/WebBrowser/Private/Android/AndroidPlatformWebBrowser.h
new file mode 100644
--- /dev/null
+++ b/Engine/Source/Runtime/WebBrowser/Private/Android/AndroidPlatformWebBrowser.h
@@ -0,0 +1,204 @@
+// Copyright 1998-2015 Epic Games, Inc. All Rights Reserved.
+
+#pragma once
+
+#if PLATFORM_ANDROID
+
+#include "IWebBrowserWindow.h"
+#include "SlateCore.h"
+
+class SAndroidWebBrowserWidget;
+
+/**
+* Implementation of interface for dealing with a Web Browser window.
+*/
+class FWebBrowserWindow
+	: public IWebBrowserWindow
+	, public TSharedFromThis<FWebBrowserWindow>
+{
+	// The WebBrowserSingleton should be the only one creating instances of this class
+	friend class FWebBrowserSingleton;
+
+private:
+	/**
+	* Creates and initializes a new instance.
+	*
+	* @param InUrl The Initial URL that will be loaded.
+	* @param InContentsToLoad Optional string to load as a web page.
+	* @param InShowErrorMessage Whether to show an error message in case of loading errors.
+	*/
+	FWebBrowserWindow(FString InUrl, TOptional<FString> InContentsToLoad, bool ShowErrorMessage, bool bThumbMouseButtonNavigation, bool bUseTransparency);
+
+public:
+	/** Virtual Destructor. */
+	virtual ~FWebBrowserWindow();
+
+public:
+
+	// IWebBrowserWindow Interface
+
+	virtual void LoadURL(FString NewURL) override;
+	virtual void LoadString(FString Contents, FString DummyURL) override;
+	virtual TSharedRef<SWidget> CreateWidget(TAttribute<FVector2D> ViewportSize) override;
+	virtual void SetViewportSize(FIntPoint WindowSize) override;
+	virtual FSlateShaderResource* GetTexture(bool bIsPopup = false) override;
+	virtual bool IsValid() const override;
+	virtual bool IsInitialized() const override;
+	virtual bool IsClosing() const override;
+	virtual EWebBrowserDocumentState GetDocumentLoadingState() const override;
+	virtual FString GetTitle() const override;
+	virtual FString GetUrl() const override;
+	virtual bool OnKeyDown(const FKeyEvent& InKeyEvent) override;
+	virtual bool OnKeyUp(const FKeyEvent& InKeyEvent) override;
+	virtual bool OnKeyChar(const FCharacterEvent& InCharacterEvent) override;
+	virtual FReply OnMouseButtonDown(const FGeometry& MyGeometry, const FPointerEvent& MouseEvent, bool bIsPopup) override;
+	virtual FReply OnMouseButtonUp(const FGeometry& MyGeometry, const FPointerEvent& MouseEvent, bool bIsPopup) override;
+	virtual FReply OnMouseButtonDoubleClick(const FGeometry& MyGeometry, const FPointerEvent& MouseEvent, bool bIsPopup) override;
+	virtual FReply OnMouseMove(const FGeometry& MyGeometry, const FPointerEvent& MouseEvent, bool bIsPopup) override;
+	virtual FReply OnMouseWheel(const FGeometry& MyGeometry, const FPointerEvent& MouseEvent, bool bIsPopup) override;
+	virtual void OnFocus(bool SetFocus, bool bIsPopup) override;
+	virtual void OnCaptureLost() override;
+	virtual bool CanGoBack() const override;
+	virtual void GoBack() override;
+	virtual bool CanGoForward() const override;
+	virtual void GoForward() override;
+	virtual bool IsLoading() const override;
+	virtual void Reload() override;
+	virtual void StopLoad() override;
+	virtual void ExecuteJavascript(const FString& Script) override;
+	virtual void CloseBrowser(bool bForce) override;
+	virtual void BindUObject(const FString& Name, UObject* Object, bool bIsPermanent = true) override;
+	virtual void UnbindUObject(const FString& Name, UObject* Object = nullptr, bool bIsPermanent = true) override;
+
+	// TODO: None of these events are actually called
+
+	DECLARE_DERIVED_EVENT(FWebBrowserWindow, IWebBrowserWindow::FOnDocumentStateChanged, FOnDocumentStateChanged);
+	virtual FOnDocumentStateChanged& OnDocumentStateChanged() override
+	{
+		return DocumentStateChangedEvent;
+	}
+
+	DECLARE_DERIVED_EVENT(FWebBrowserWindow, IWebBrowserWindow::FOnTitleChanged, FOnTitleChanged);
+	virtual FOnTitleChanged& OnTitleChanged() override
+	{
+		return TitleChangedEvent;
+	}
+
+	DECLARE_DERIVED_EVENT(FWebBrowserWindow, IWebBrowserWindow::FOnUrlChanged, FOnUrlChanged);
+	virtual FOnUrlChanged& OnUrlChanged() override
+	{
+		return UrlChangedEvent;
+	}
+
+	DECLARE_DERIVED_EVENT(FWebBrowserWindow, IWebBrowserWindow::FOnNeedsRedraw, FOnNeedsRedraw);
+	virtual FOnNeedsRedraw& OnNeedsRedraw() override
+	{
+		return NeedsRedrawEvent;
+	}
+
+	virtual FOnBeforeBrowse& OnBeforeBrowse() override
+	{
+		return BeforeBrowseDelegate;
+	}
+
+	virtual FOnLoadUrl& OnLoadUrl() override
+	{
+		return LoadUrlDelegate;
+	}
+
+	virtual FOnCreateWindow& OnCreateWindow() override
+	{
+		return CreateWindowDelegate;
+	}
+
+	virtual FOnCloseWindow& OnCloseWindow() override
+	{
+		return CloseWindowDelegate;
+	}
+
+	virtual FCursorReply OnCursorQuery(const FGeometry& MyGeometry, const FPointerEvent& CursorEvent) override
+	{
+		//return Cursor == EMouseCursor::Default ? FCursorReply::Unhandled() : FCursorReply::Cursor(Cursor);
+		return FCursorReply::Unhandled();
+	}
+
+	virtual FOnBeforePopupDelegate& OnBeforePopup() override
+	{
+		return BeforePopupDelegate;
+	}
+
+	DECLARE_DERIVED_EVENT(FWebBrowserWindow, IWebBrowserWindow::FOnShowPopup, FOnShowPopup);
+	virtual FOnShowPopup& OnShowPopup() override
+	{
+		return ShowPopupEvent;
+	}
+
+	DECLARE_DERIVED_EVENT(FWebBrowserWindow, IWebBrowserWindow::FOnDismissPopup, FOnDismissPopup);
+	virtual FOnDismissPopup& OnDismissPopup() override
+	{
+		return DismissPopupEvent;
+	}
+
+public:
+
+	/**
+	* Called from the WebBrowserSingleton tick event. Should test whether the widget got a tick from Slate last frame and set the state to hidden if not.
+	*/
+	//void CheckTickActivity();
+
+private:
+
+	TSharedPtr<SAndroidWebBrowserWidget> BrowserWidget;
+
+	/** Current Url of this window. */
+	FString CurrentUrl;
+
+	/** Optional text to load as a web page. */
+	TOptional<FString> ContentsToLoad;
+
+	// TODO: None of these events are actually called
+
+	/** Delegate for broadcasting load state changes. */
+	FOnDocumentStateChanged DocumentStateChangedEvent;
+
+	/** Delegate for broadcasting title changes. */
+	FOnTitleChanged TitleChangedEvent;
+
+	/** Delegate for broadcasting address changes. */
+	FOnUrlChanged UrlChangedEvent;
+
+	/** Delegate for notifying that the window needs refreshing. */
+	FOnNeedsRedraw NeedsRedrawEvent;
+
+	/** Delegate that is executed prior to browser navigation. */
+	FOnBeforeBrowse BeforeBrowseDelegate;
+
+	/** Delegate for overriding Url contents. */
+	FOnLoadUrl LoadUrlDelegate;
+
+	/** Delegate for notifying that a popup window is attempting to open. */
+	FOnBeforePopupDelegate BeforePopupDelegate;
+
+	/** Delegate for handling requests to create new windows. */
+	FOnCreateWindow CreateWindowDelegate;
+
+	/** Delegate for handling requests to close new windows that were created. */
+	FOnCloseWindow CloseWindowDelegate;
+
+	/** Delegate for handling requests to show the popup menu. */
+	FOnShowPopup ShowPopupEvent;
+
+	/** Delegate for handling requests to dismiss the current popup menu. */
+	FOnDismissPopup DismissPopupEvent;
+
+	///** Tracks wether the widget is currently hidden or not*/
+	//bool bIsHidden;
+
+	///** Used to detect when the widget is hidden*/
+	//bool bTickedLastFrame;
+
+	///** Handling of passing and marshalling messages for JS integration is delegated to a helper class*/
+	////TSharedPtr<FWebJSScripting> Scripting;
+};
+
+#endif
diff --git a/Engine/Source/Runtime/WebBrowser/Private/SWebBrowser.cpp b/Engine/Source/Runtime/WebBrowser/Private/SWebBrowser.cpp
--- a/Engine/Source/Runtime/WebBrowser/Private/SWebBrowser.cpp
+++ b/Engine/Source/Runtime/WebBrowser/Private/SWebBrowser.cpp
@@ -19,13 +19,12 @@ SWebBrowser::~SWebBrowser()
 {
 	if (BrowserWindow.IsValid())
 	{
-		BrowserWindow->OnCreateWindow().Unbind();		
+		BrowserWindow->OnCreateWindow().Unbind();
 		BrowserWindow->OnCloseWindow().Unbind();
 		BrowserWindow->OnDocumentStateChanged().RemoveAll(this);
 		BrowserWindow->OnNeedsRedraw().RemoveAll(this);
 		BrowserWindow->OnTitleChanged().RemoveAll(this);
 		BrowserWindow->OnUrlChanged().RemoveAll(this);
-		BrowserWindow->OnToolTip().RemoveAll(this);
 		BrowserWindow->OnShowPopup().RemoveAll(this);
 		BrowserWindow->OnDismissPopup().RemoveAll(this);
 		BrowserWindow->OnBeforeBrowse().Unbind();
@@ -68,9 +67,6 @@ void SWebBrowser::Construct(const FArguments& InArgs, const TSharedPtr<IWebBrows
 			OSWindowHandle = NativeWindow->GetOSWindowHandle();
 		}
 
-		static bool AllowCEF = !FParse::Param(FCommandLine::Get(), TEXT("nocef"));
-		if (AllowCEF)
-		{
 		BrowserWindow = IWebBrowserModule::Get().GetSingleton()->CreateBrowserWindow(
 			OSWindowHandle,
 			InArgs._InitialURL,
@@ -79,10 +75,13 @@ void SWebBrowser::Construct(const FArguments& InArgs, const TSharedPtr<IWebBrows
 			InArgs._ContentsToLoad,
 			InArgs._ShowErrorMessage,
 			InArgs._BackgroundColor
-			);		
-		}
+			);
 	}
 
+	TSharedRef<SWidget> BrowserWidgetRef = BrowserWindow.IsValid() ? BrowserWindow->CreateWidget(InArgs._ViewportSize) : SNullWidget::NullWidget;
+	BrowserWidgetRef->SetVisibility(TAttribute<EVisibility>(this, &SWebBrowser::GetViewportVisibility));
+	BrowserWidget = BrowserWidgetRef;
+
 	ChildSlot
 	[
 		SNew(SVerticalBox)
@@ -150,16 +149,11 @@ void SWebBrowser::Construct(const FArguments& InArgs, const TSharedPtr<IWebBrows
 		+SVerticalBox::Slot()
 		[
 			SNew(SOverlay)
-			+SOverlay::Slot()
+			+ SOverlay::Slot()
 			[
-				SAssignNew(ViewportWidget, SViewport)
-				.ViewportSize(InArgs._ViewportSize)
-				.EnableGammaCorrection(false)
-				.EnableBlending(InArgs._SupportsTransparency)
-				.IgnoreTextureAlpha(!InArgs._SupportsTransparency)
-				.Visibility(this, &SWebBrowser::GetViewportVisibility)
+				BrowserWidgetRef
 			]
-			+SOverlay::Slot()
+			+ SOverlay::Slot()
 			.HAlign(HAlign_Center)
 			.VAlign(VAlign_Center)
 			[
@@ -175,7 +169,7 @@ void SWebBrowser::Construct(const FArguments& InArgs, const TSharedPtr<IWebBrows
 	{
 		if(OnCreateWindow.IsBound())
 		{
-			BrowserWindow->OnCreateWindow().BindSP(this, &SWebBrowser::HandleCreateWindow);		
+			BrowserWindow->OnCreateWindow().BindSP(this, &SWebBrowser::HandleCreateWindow);
 		}
 
 		if(OnCloseWindow.IsBound())
@@ -187,7 +181,6 @@ void SWebBrowser::Construct(const FArguments& InArgs, const TSharedPtr<IWebBrows
 		BrowserWindow->OnNeedsRedraw().AddSP(this, &SWebBrowser::HandleBrowserWindowNeedsRedraw);
 		BrowserWindow->OnTitleChanged().AddSP(this, &SWebBrowser::HandleTitleChanged);
 		BrowserWindow->OnUrlChanged().AddSP(this, &SWebBrowser::HandleUrlChanged);
-		BrowserWindow->OnToolTip().AddSP(this, &SWebBrowser::HandleToolTip);
 		BrowserWindow->OnBeforeBrowse().BindSP(this, &SWebBrowser::HandleBeforeNavigation);
 		BrowserWindow->OnLoadUrl().BindSP(this, &SWebBrowser::HandleLoadUrl);
 		BrowserWindow->OnShowDialog().BindSP(this, &SWebBrowser::HandleShowDialog);
@@ -195,8 +188,6 @@ void SWebBrowser::Construct(const FArguments& InArgs, const TSharedPtr<IWebBrows
 		BrowserWindow->OnBeforePopup().BindSP(this, &SWebBrowser::HandleBeforePopup);
 		BrowserWindow->OnShowPopup().AddSP(this, &SWebBrowser::HandleShowPopup);
 		BrowserWindow->OnDismissPopup().AddSP(this, &SWebBrowser::HandleDismissPopup);
-		BrowserViewport = MakeShareable(new FWebBrowserViewport(BrowserWindow));
-		ViewportWidget->SetViewportInterface(BrowserViewport.ToSharedRef());
 	}
 	else
 	{
@@ -429,20 +420,6 @@ void SWebBrowser::HandleUrlChanged( FString NewUrl )
 	OnUrlChanged.ExecuteIfBound(AddressBarUrl);
 }
 
-void SWebBrowser::HandleToolTip(FString ToolTipText)
-{
-	if(ToolTipText.IsEmpty())
-	{
-		FSlateApplication::Get().CloseToolTip();
-		ViewportWidget->SetToolTip(nullptr);
-	}
-	else
-	{
-		ViewportWidget->SetToolTipText(FText::FromString(ToolTipText));
-		FSlateApplication::Get().UpdateToolTip(true);
-	}
-}
-
 bool SWebBrowser::HandleBeforeNavigation(const FString& Url, bool bIsRedirect)
 {
 	if(OnBeforeNavigation.IsBound())
@@ -560,7 +537,7 @@ void SWebBrowser::HandleShowPopup(const FIntRect& PopupSize)
 	MenuViewport = MakeShareable(new FWebBrowserViewport(BrowserWindow, true));
 	MenuContent->SetViewportInterface(MenuViewport.ToSharedRef());
 	FWidgetPath WidgetPath;
-	FSlateApplication::Get().GeneratePathToWidgetUnchecked(ViewportWidget.ToSharedRef(), WidgetPath);
+	FSlateApplication::Get().GeneratePathToWidgetUnchecked(BrowserWidget.ToSharedRef(), WidgetPath);
 	if (WidgetPath.IsValid())
 	{
 		TSharedRef< SWidget > MenuContentRef = MenuContent.ToSharedRef();
@@ -569,7 +546,7 @@ void SWebBrowser::HandleShowPopup(const FIntRect& PopupSize)
 
 
 		// Open the pop-up. The popup method will be queried from the widget path passed in.
-		TSharedPtr<IMenu> NewMenu = FSlateApplication::Get().PushMenu(ViewportWidget.ToSharedRef(), WidgetPath, MenuContentRef, NewPosition, FPopupTransitionEffect( FPopupTransitionEffect::ComboButton ), false);
+		TSharedPtr<IMenu> NewMenu = FSlateApplication::Get().PushMenu(BrowserWidget.ToSharedRef(), WidgetPath, MenuContentRef, NewPosition, FPopupTransitionEffect( FPopupTransitionEffect::ComboButton ), false);
 		NewMenu->GetOnMenuDismissed().AddSP(this, &SWebBrowser::HandleMenuDismissed);
 		PopupMenuPtr = NewMenu;
 	}
diff --git a/Engine/Source/Runtime/WebBrowser/Private/WebBrowserSingleton.cpp b/Engine/Source/Runtime/WebBrowser/Private/WebBrowserSingleton.cpp
--- a/Engine/Source/Runtime/WebBrowser/Private/WebBrowserSingleton.cpp
+++ b/Engine/Source/Runtime/WebBrowser/Private/WebBrowserSingleton.cpp
@@ -28,23 +28,28 @@
 #	include <pthread.h>
 #endif
 
-// Define some platform-dependent file locations
-#define CEF3_BIN_DIR TEXT("Binaries/ThirdParty/CEF3")
-#if PLATFORM_WINDOWS && PLATFORM_64BITS
-#	define CEF3_RESOURCES_DIR CEF3_BIN_DIR TEXT("/Win64/Resources")
-#	define CEF3_SUBPROCES_EXE TEXT("Binaries/Win64/UnrealCEFSubProcess.exe")
-#elif PLATFORM_WINDOWS && PLATFORM_32BITS
-#	define CEF3_RESOURCES_DIR CEF3_BIN_DIR TEXT("/Win32/Resources")
-#	define CEF3_SUBPROCES_EXE TEXT("Binaries/Win32/UnrealCEFSubProcess.exe")
-#elif PLATFORM_MAC
-#	define CEF3_FRAMEWORK_DIR CEF3_BIN_DIR TEXT("/Mac/Chromium Embedded Framework.framework")
-#	define CEF3_RESOURCES_DIR CEF3_FRAMEWORK_DIR TEXT("/Resources")
-#	define CEF3_SUBPROCES_EXE TEXT("Binaries/Mac/UnrealCEFSubProcess.app/Contents/MacOS/UnrealCEFSubProcess")
-#elif PLATFORM_LINUX // @todo Linux
-#	define CEF3_RESOURCES_DIR CEF3_BIN_DIR TEXT("/Linux/Resources")
-#	define CEF3_SUBPROCES_EXE TEXT("Binaries/Linux/UnrealCEFSubProcess")
+#if PLATFORM_ANDROID
+#	include <Android/AndroidPlatformWebBrowser.h>
 #endif
 
+// Define some platform-dependent file locations
+#if WITH_CEF3
+#	define CEF3_BIN_DIR TEXT("Binaries/ThirdParty/CEF3")
+#	if PLATFORM_WINDOWS && PLATFORM_64BITS
+#		define CEF3_RESOURCES_DIR CEF3_BIN_DIR TEXT("/Win64/Resources")
+#		define CEF3_SUBPROCES_EXE TEXT("Binaries/Win64/UnrealCEFSubProcess.exe")
+#	elif PLATFORM_WINDOWS && PLATFORM_32BITS
+#		define CEF3_RESOURCES_DIR CEF3_BIN_DIR TEXT("/Win32/Resources")
+#		define CEF3_SUBPROCES_EXE TEXT("Binaries/Win32/UnrealCEFSubProcess.exe")
+#	elif PLATFORM_MAC
+#		define CEF3_FRAMEWORK_DIR CEF3_BIN_DIR TEXT("/Mac/Chromium Embedded Framework.framework")
+#		define CEF3_RESOURCES_DIR CEF3_FRAMEWORK_DIR TEXT("/Resources")
+#		define CEF3_SUBPROCES_EXE TEXT("Binaries/Mac/UnrealCEFSubProcess.app/Contents/MacOS/UnrealCEFSubProcess")
+#	elif PLATFORM_LINUX // @todo Linux
+#		define CEF3_RESOURCES_DIR CEF3_BIN_DIR TEXT("/Linux/Resources")
+#		define CEF3_SUBPROCES_EXE TEXT("Binaries/Linux/UnrealCEFSubProcess")
+#	endif
+#endif
 
 namespace {
 
@@ -264,39 +269,48 @@ TSharedPtr<IWebBrowserWindow> FWebBrowserSingleton::CreateBrowserWindow(
 	FColor BackgroundColor)
 {
 #if WITH_CEF3
+	static bool AllowCEF = !FParse::Param(FCommandLine::Get(), TEXT("nocef"));
+	if (AllowCEF)
+	{
+		// Information used when creating the native window.
+		CefWindowHandle WindowHandle = (CefWindowHandle)OSWindowHandle; // TODO: check this is correct for all platforms
+		CefWindowInfo WindowInfo;
 
-	// Information used when creating the native window.
-	CefWindowHandle WindowHandle = (CefWindowHandle)OSWindowHandle; // TODO: check this is correct for all platforms
-	CefWindowInfo WindowInfo;
+		// Always use off screen rendering so we can integrate with our windows
+		WindowInfo.SetAsWindowless(WindowHandle, bUseTransparency);
 
-	// Always use off screen rendering so we can integrate with our windows
-	WindowInfo.SetAsWindowless(WindowHandle, bUseTransparency);
+		// Specify CEF browser settings here.
+		CefBrowserSettings BrowserSettings;
 
-	// Specify CEF browser settings here.
-	CefBrowserSettings BrowserSettings;
-	
-	// Set max framerate to maximum supported.
-	BrowserSettings.windowless_frame_rate = 60;
-	BrowserSettings.background_color = CefColorSetARGB(BackgroundColor.A, BackgroundColor.R, BackgroundColor.G, BackgroundColor.B);
+		// Set max framerate to maximum supported.
+		BrowserSettings.windowless_frame_rate = 60;
+		BrowserSettings.background_color = CefColorSetARGB(BackgroundColor.A, BackgroundColor.R, BackgroundColor.G, BackgroundColor.B);
 
-	// Disable plugins
-	BrowserSettings.plugins = STATE_DISABLED;
+		// Disable plugins
+		BrowserSettings.plugins = STATE_DISABLED;
 
 
-	// WebBrowserHandler implements browser-level callbacks.
-	CefRefPtr<FWebBrowserHandler> NewHandler(new FWebBrowserHandler);
+		// WebBrowserHandler implements browser-level callbacks.
+		CefRefPtr<FWebBrowserHandler> NewHandler(new FWebBrowserHandler);
 
-	// Create the CEF browser window.
-	CefRefPtr<CefBrowser> Browser = CefBrowserHost::CreateBrowserSync(WindowInfo, NewHandler.get(), *InitialURL, BrowserSettings, nullptr);
-	if (Browser.get())
-	{
-		// Create new window
-		TSharedPtr<FWebBrowserWindow> NewBrowserWindow(new FWebBrowserWindow(Browser, InitialURL, ContentsToLoad, ShowErrorMessage, bThumbMouseButtonNavigation, bUseTransparency));
-		NewHandler->SetBrowserWindow(NewBrowserWindow);
+		// Create the CEF browser window.
+		CefRefPtr<CefBrowser> Browser = CefBrowserHost::CreateBrowserSync(WindowInfo, NewHandler.get(), *InitialURL, BrowserSettings, nullptr);
+		if (Browser.get())
+		{
+			// Create new window
+			TSharedPtr<FWebBrowserWindow> NewBrowserWindow(new FWebBrowserWindow(Browser, InitialURL, ContentsToLoad, ShowErrorMessage, bThumbMouseButtonNavigation, bUseTransparency));
+			NewHandler->SetBrowserWindow(NewBrowserWindow);
 
-		WindowInterfaces.Add(NewBrowserWindow);
-		return NewBrowserWindow;
+			WindowInterfaces.Add(NewBrowserWindow);
+			return NewBrowserWindow;
+		}
 	}
+#elif PLATFORM_ANDROID
+	// Create new window
+	TSharedPtr<FWebBrowserWindow> NewBrowserWindow = MakeShareable(new FWebBrowserWindow(InitialURL, ContentsToLoad, ShowErrorMessage, bThumbMouseButtonNavigation, bUseTransparency));
+
+	//WindowInterfaces.Add(NewBrowserWindow);
+	return NewBrowserWindow;
 #endif
 	return nullptr;
 }
diff --git a/Engine/Source/Runtime/WebBrowser/Private/WebBrowserWindow.cpp b/Engine/Source/Runtime/WebBrowser/Private/WebBrowserWindow.cpp
--- a/Engine/Source/Runtime/WebBrowser/Private/WebBrowserWindow.cpp
+++ b/Engine/Source/Runtime/WebBrowser/Private/WebBrowserWindow.cpp
@@ -2,6 +2,7 @@
 
 #include "WebBrowserPrivatePCH.h"
 #include "WebBrowserWindow.h"
+#include "WebBrowserViewport.h"
 #include "WebBrowserByteResource.h"
 #include "WebBrowserPopupFeatures.h"
 #include "WebBrowserDialog.h"
@@ -117,6 +118,23 @@ void FWebBrowserWindow::SetViewportSize(FIntPoint WindowSize)
 	}
 }
 
+TSharedRef<SWidget> FWebBrowserWindow::CreateWidget(TAttribute<FVector2D> InViewportSize)
+{
+	TSharedRef<SViewport> ViewportWidgetRef =
+		SNew(SViewport)
+		.ViewportSize(InViewportSize)
+		.EnableGammaCorrection(false)
+		.EnableBlending(bUseTransparency)
+		.IgnoreTextureAlpha(!bUseTransparency);
+
+	TSharedRef<FWebBrowserViewport> BrowserViewportRef = MakeShareable(new FWebBrowserViewport(this->AsShared(), ViewportWidgetRef));
+	BrowserViewport = BrowserViewportRef;
+	ViewportWidgetRef->SetViewportInterface(MoveTemp(BrowserViewportRef));
+	ViewportWidget = ViewportWidgetRef;
+
+	return ViewportWidgetRef;
+}
+
 FSlateShaderResource* FWebBrowserWindow::GetTexture(bool bIsPopup)
 {
 	if (UpdatableTextures[bIsPopup?1:0] != nullptr && IsInitialized())
@@ -772,7 +790,17 @@ void FWebBrowserWindow::SetToolTip(const CefString& CefToolTip)
 	if (ToolTipText != NewToolTipText)
 	{
 		ToolTipText = NewToolTipText;
-		OnToolTip().Broadcast(ToolTipText);
+
+		if (ToolTipText.IsEmpty())
+		{
+			FSlateApplication::Get().CloseToolTip();
+			ViewportWidget->SetToolTip(nullptr);
+		}
+		else
+		{
+			ViewportWidget->SetToolTipText(FText::FromString(ToolTipText));
+			FSlateApplication::Get().UpdateToolTip(true);
+		}
 	}
 }
 
diff --git a/Engine/Source/Runtime/WebBrowser/Private/WebBrowserWindow.h b/Engine/Source/Runtime/WebBrowser/Private/WebBrowserWindow.h
--- a/Engine/Source/Runtime/WebBrowser/Private/WebBrowserWindow.h
+++ b/Engine/Source/Runtime/WebBrowser/Private/WebBrowserWindow.h
@@ -26,6 +26,8 @@
 	#include "HideWindowsPlatformTypes.h"
 #endif
 
+class FWebBrowserViewport;
+
 /**
  * Helper for containing items required for CEF browser window creation.
  */
@@ -78,6 +80,7 @@ public:
 	virtual void LoadURL(FString NewURL) override;
 	virtual void LoadString(FString Contents, FString DummyURL) override;
 	virtual void SetViewportSize(FIntPoint WindowSize) override;
+	virtual TSharedRef<SWidget> CreateWidget(TAttribute<FVector2D> ViewportSize) override;
 	virtual FSlateShaderResource* GetTexture(bool bIsPopup = false) override;
 	virtual bool IsValid() const override;
 	virtual bool IsInitialized() const override;
@@ -125,12 +128,6 @@ public:
 		return UrlChangedEvent;
 	}
 
-	DECLARE_DERIVED_EVENT(FWebBrowserWindow, IWebBrowserWindow::FOnToolTip, FOnToolTip);
-	virtual FOnToolTip& OnToolTip() override
-	{
-		return ToolTipEvent;
-	}
-
 	DECLARE_DERIVED_EVENT(FWebBrowserWindow, IWebBrowserWindow::FOnNeedsRedraw, FOnNeedsRedraw);
 	virtual FOnNeedsRedraw& OnNeedsRedraw() override
 	{
@@ -399,6 +396,11 @@ private:
 
 private:
 
+	/** Viewport interface for rendering the web page. */
+	TSharedPtr<FWebBrowserViewport> BrowserViewport;
+	/** The actual viewport widget. Required to update its tool tip property. */
+	TSharedPtr<SViewport> ViewportWidget;
+
 	/** Current state of the document being loaded. */
 	EWebBrowserDocumentState DocumentState;
 
@@ -447,9 +449,6 @@ private:
 	/** Delegate for broadcasting address changes. */
 	FOnUrlChanged UrlChangedEvent;
 
-	/** Delegate for showing or hiding tool tips. */
-	FOnToolTip ToolTipEvent;
-
 	/** Delegate for notifying that the window needs refreshing. */
 	FOnNeedsRedraw NeedsRedrawEvent;
 
diff --git a/Engine/Source/Runtime/WebBrowser/Public/IWebBrowserWindow.h b/Engine/Source/Runtime/WebBrowser/Public/IWebBrowserWindow.h
--- a/Engine/Source/Runtime/WebBrowser/Public/IWebBrowserWindow.h
+++ b/Engine/Source/Runtime/WebBrowser/Public/IWebBrowserWindow.h
@@ -12,6 +12,7 @@ class FReply;
 class FCursorReply;
 class FSlateShaderResource;
 class IWebBrowserPopupFeatures;
+class SWidget;
 
 
 enum class EWebBrowserDocumentState
@@ -44,6 +45,13 @@ public:
 	 */
 	virtual void LoadString(FString Contents, FString DummyURL) = 0;
 
+	/**
+	 * Create the SWidget for this WebBrowser
+	 *
+	 * @param ViewportSize Size of viewport for CEF implementation on PC
+	 */
+	virtual TSharedRef<SWidget> CreateWidget(TAttribute<FVector2D> ViewportSize) = 0;
+
 	/**
 	 * Set the desired size of the web browser viewport
 	 * 
@@ -253,10 +261,6 @@ public:
 	DECLARE_EVENT_OneParam(IWebBrowserWindow, FOnUrlChanged, FString /*NewUrl*/);
 	virtual FOnUrlChanged& OnUrlChanged() = 0;
 
-	/** A delegate to allow callbacks when a frame url changes. */
-	DECLARE_EVENT_OneParam(IWebBrowserWindow, FOnToolTip, FString /*ToolTipText*/);
-	virtual FOnToolTip& OnToolTip() = 0;
-
 	/** A delegate that is invoked when the off-screen window has been repainted and requires an update. */
 	DECLARE_EVENT(IWebBrowserWindow, FOnNeedsRedraw)
 	virtual FOnNeedsRedraw& OnNeedsRedraw() = 0;
diff --git a/Engine/Source/Runtime/WebBrowser/Public/SWebBrowser.h b/Engine/Source/Runtime/WebBrowser/Public/SWebBrowser.h
--- a/Engine/Source/Runtime/WebBrowser/Public/SWebBrowser.h
+++ b/Engine/Source/Runtime/WebBrowser/Public/SWebBrowser.h
@@ -244,9 +244,6 @@ private:
 	/** Callback for loaded url changes. */
 	void HandleUrlChanged(FString NewUrl);
 	
-	/** Callback for showing browser tool tips. */
-	void HandleToolTip(FString ToolTipText);
-
 	/**
 	 * A delegate that is executed prior to browser navigation.
 	 *
@@ -295,14 +292,11 @@ private:
 	/** Interface for dealing with a web browser window. */
 	TSharedPtr<IWebBrowserWindow> BrowserWindow;
 
-	/** Viewport interface for rendering the web page. */
-	TSharedPtr<FWebBrowserViewport> BrowserViewport;
+	TSharedPtr<SWidget> BrowserWidget;
+
 	/** Viewport interface for rendering popup menus. */
 	TSharedPtr<FWebBrowserViewport>	MenuViewport;
 
-	/** The actual viewport widget. Required to update its tool tip property. */
-	TSharedPtr<SViewport> ViewportWidget;
-
 	TArray<TSharedRef<IWebBrowserAdapter>> Adapters;
 
 	/**
