Fix bad merge

diff --git a/Engine/Source/Runtime/WebBrowser/Private/SWebBrowser.cpp b/Engine/Source/Runtime/WebBrowser/Private/SWebBrowser.cpp
--- a/Engine/Source/Runtime/WebBrowser/Private/SWebBrowser.cpp
+++ b/Engine/Source/Runtime/WebBrowser/Private/SWebBrowser.cpp
@@ -571,7 +571,7 @@ void SWebBrowser::HandleDismissPopup()
 	if (PopupMenuPtr.IsValid())
 	{
 		PopupMenuPtr.Pin()->Dismiss();
-		FSlateApplication::Get().SetKeyboardFocus(ViewportWidget, EFocusCause::SetDirectly);
+		FSlateApplication::Get().SetKeyboardFocus(BrowserWidget, EFocusCause::SetDirectly);
 	}
 }
 
diff --git a/Engine/Source/Runtime/WebBrowser/Private/WebBrowserWindow.cpp b/Engine/Source/Runtime/WebBrowser/Private/WebBrowserWindow.cpp
--- a/Engine/Source/Runtime/WebBrowser/Private/WebBrowserWindow.cpp
+++ b/Engine/Source/Runtime/WebBrowser/Private/WebBrowserWindow.cpp
@@ -127,7 +127,7 @@ TSharedRef<SWidget> FWebBrowserWindow::CreateWidget(TAttribute<FVector2D> InView
 		.EnableBlending(bUseTransparency)
 		.IgnoreTextureAlpha(!bUseTransparency);
 
-	TSharedRef<FWebBrowserViewport> BrowserViewportRef = MakeShareable(new FWebBrowserViewport(this->AsShared(), ViewportWidgetRef));
+	TSharedRef<FWebBrowserViewport> BrowserViewportRef = MakeShareable(new FWebBrowserViewport(this->AsShared()));
 	BrowserViewport = BrowserViewportRef;
 	ViewportWidgetRef->SetViewportInterface(MoveTemp(BrowserViewportRef));
 	ViewportWidget = ViewportWidgetRef;
