More merge fixes

diff --git a/Engine/Source/Runtime/WebBrowser/Private/Android/AndroidPlatformWebBrowser.h b/Engine/Source/Runtime/WebBrowser/Private/Android/AndroidPlatformWebBrowser.h
--- a/Engine/Source/Runtime/WebBrowser/Private/Android/AndroidPlatformWebBrowser.h
+++ b/Engine/Source/Runtime/WebBrowser/Private/Android/AndroidPlatformWebBrowser.h
@@ -139,6 +139,16 @@ public:
 		return DismissPopupEvent;
 	}
 
+	virtual FOnShowDialog& OnShowDialog() override
+	{
+		return ShowDialogDelegate;
+	}
+
+	virtual FOnDismissAllDialogs& OnDismissAllDialogs() override
+	{
+		return DismissAllDialogsDelegate;
+	}
+
 public:
 
 	/**
@@ -191,6 +201,12 @@ private:
 	/** Delegate for handling requests to dismiss the current popup menu. */
 	FOnDismissPopup DismissPopupEvent;
 
+	/** Delegate for showing dialogs. */
+	FOnShowDialog ShowDialogDelegate;
+
+	/** Delegate for dismissing all dialogs. */
+	FOnDismissAllDialogs DismissAllDialogsDelegate;
+
 	///** Tracks wether the widget is currently hidden or not*/
 	//bool bIsHidden;
 
