WebBrowser: Automatically reload the current page if render subprocess is killed or crashes.
OPP-4409
Merging CL#2717757 using UE4-To-UE4-LauncherDev

diff --git a/Engine/Source/Runtime/WebBrowser/Private/WebBrowserHandler.cpp b/Engine/Source/Runtime/WebBrowser/Private/WebBrowserHandler.cpp
--- a/Engine/Source/Runtime/WebBrowser/Private/WebBrowserHandler.cpp
+++ b/Engine/Source/Runtime/WebBrowser/Private/WebBrowserHandler.cpp
@@ -300,6 +300,11 @@ bool FWebBrowserHandler::OnBeforeResourceLoad(CefRefPtr<CefBrowser> Browser, Cef
 
 void FWebBrowserHandler::OnRenderProcessTerminated(CefRefPtr<CefBrowser> Browser, TerminationStatus Status)
 {
+	TSharedPtr<FWebBrowserWindow> BrowserWindow = BrowserWindowPtr.Pin();
+	if (BrowserWindow.IsValid())
+	{
+		BrowserWindow->OnRenderProcessTerminated(Status);
+	}
 }
 
 bool FWebBrowserHandler::OnBeforeBrowse(CefRefPtr<CefBrowser> Browser,
diff --git a/Engine/Source/Runtime/WebBrowser/Private/WebBrowserWindow.cpp b/Engine/Source/Runtime/WebBrowser/Private/WebBrowserWindow.cpp
--- a/Engine/Source/Runtime/WebBrowser/Private/WebBrowserWindow.cpp
+++ b/Engine/Source/Runtime/WebBrowser/Private/WebBrowserWindow.cpp
@@ -547,6 +547,17 @@ void FWebBrowserWindow::OnResetDialogState()
 	OnDismissAllDialogs().ExecuteIfBound();
 }
 
+void FWebBrowserWindow::OnRenderProcessTerminated(CefRequestHandler::TerminationStatus Status)
+{
+	if(bRecoverFromRenderProcessCrash)
+	{
+		bRecoverFromRenderProcessCrash = false;
+		NotifyDocumentError(); // Only attempt a single recovery at a time
+	}
+
+	bRecoverFromRenderProcessCrash = true;
+	Reload();
+}
 
 FReply FWebBrowserWindow::OnMouseButtonDown(const FGeometry& MyGeometry, const FPointerEvent& MouseEvent, bool bIsPopup)
 {
@@ -829,6 +840,14 @@ void FWebBrowserWindow::NotifyDocumentLoadingStateChange(bool IsLoading)
 	if (! IsLoading)
 	{
 		bIsInitialized = true;
+
+		if (bRecoverFromRenderProcessCrash)
+		{
+			bRecoverFromRenderProcessCrash = false;
+			// Toggle hidden/visible state to get OnPaint calls from CEF.
+			SetIsHidden(true);
+			SetIsHidden(false);
+		}
 	}
 	
 	EWebBrowserDocumentState NewState = IsLoading
diff --git a/Engine/Source/Runtime/WebBrowser/Private/WebBrowserWindow.h b/Engine/Source/Runtime/WebBrowser/Private/WebBrowserWindow.h
--- a/Engine/Source/Runtime/WebBrowser/Private/WebBrowserWindow.h
+++ b/Engine/Source/Runtime/WebBrowser/Private/WebBrowserWindow.h
@@ -298,6 +298,11 @@ private:
 	 */
 	void OnResetDialogState();
 
+	/**
+	 * Called when render process was terminated abnormally.
+	 */
+	void OnRenderProcessTerminated(CefRequestHandler::TerminationStatus Status);
+
 
 	/** Specifies if window creation functionality is available. 
 	 *
@@ -503,6 +508,9 @@ private:
 	FIntPoint PopupPosition;
 	bool bShowPopupRequested;
 
+	/** This is set to true when reloading after render process crash. */
+	bool bRecoverFromRenderProcessCrash;
+
 	/** Handling of passing and marshalling messages for JS integration is delegated to a helper class*/
 	TSharedPtr<FWebJSScripting> Scripting;
 };
