WebBrowser: Set CEF log file path explicitly to a path under GameLogDir, as by default it is placed next to the executable.
OPP-3920: Chromium integration creates log file next to UnrealCEFSubProcess.exe
Merging CL#2717950 using UE4-To-UE4-LauncherDev

diff --git a/Engine/Source/Runtime/WebBrowser/Private/WebBrowserSingleton.cpp b/Engine/Source/Runtime/WebBrowser/Private/WebBrowserSingleton.cpp
--- a/Engine/Source/Runtime/WebBrowser/Private/WebBrowserSingleton.cpp
+++ b/Engine/Source/Runtime/WebBrowser/Private/WebBrowserSingleton.cpp
@@ -139,6 +139,10 @@ FWebBrowserSingleton::FWebBrowserSingleton()
 	CefSettings Settings;
 	Settings.no_sandbox = true;
 	Settings.command_line_args_disabled = true;
+
+	FString CefLogFile(FPaths::Combine(*FPaths::GameLogDir(), TEXT("cef3.log")));
+	CefLogFile = FPaths::ConvertRelativePathToFull(CefLogFile);
+	CefString(&Settings.log_file) = *CefLogFile;
 	Settings.log_severity = LOGSEVERITY_WARNING;
 
 	// Specify locale from our settings
