fix for tvOS usage of prohibited UIWebView class

diff --git a/Engine/Source/Runtime/WebBrowser/Private/IOS/IOSPlatformWebBrowser.cpp b/Engine/Source/Runtime/WebBrowser/Private/IOS/IOSPlatformWebBrowser.cpp
--- a/Engine/Source/Runtime/WebBrowser/Private/IOS/IOSPlatformWebBrowser.cpp
+++ b/Engine/Source/Runtime/WebBrowser/Private/IOS/IOSPlatformWebBrowser.cpp
@@ -11,7 +11,9 @@
 
 @implementation IOSWebViewWrapper
 
+#if !PLATFORM_TVOS
 @synthesize WebView;
+#endif
 @synthesize NextURL;
 @synthesize NextContent;
 
@@ -22,27 +24,32 @@
 	NextContent = nil;
 	bNeedsAddToView = true;
 
+#if !PLATFORM_TVOS
 	dispatch_async(dispatch_get_main_queue(), ^
 	{
 		WebView = [[UIWebView alloc]initWithFrame:CGRectMake(1, 1, 100, 100)];
 		WebView.delegate = self;
 	});
+#endif
 }
 
 -(void)close;
 {
+#if !PLATFORM_TVOS
 	WebView.delegate = nil;
 	dispatch_async(dispatch_get_main_queue(), ^
 	{
 		[self.WebView removeFromSuperview];
 		WebView = nil;
 	});
+#endif
 }
 
 -(void)updateframe:(CGRect)InFrame;
 {
 	self.DesiredFrame = InFrame;
 
+#if !PLATFORM_TVOS
 	dispatch_async(dispatch_get_main_queue(), ^
 	{
 		if(WebView != nil)
@@ -73,6 +80,7 @@
 			}
 		}
 	});
+#endif
 }
 
 -(void)loadurl:(NSURL*)InURL;
diff --git a/Engine/Source/Runtime/WebBrowser/Private/IOS/IOSPlatformWebBrowser.h b/Engine/Source/Runtime/WebBrowser/Private/IOS/IOSPlatformWebBrowser.h
--- a/Engine/Source/Runtime/WebBrowser/Private/IOS/IOSPlatformWebBrowser.h
+++ b/Engine/Source/Runtime/WebBrowser/Private/IOS/IOSPlatformWebBrowser.h
@@ -15,13 +15,18 @@ class SWebBrowserView;
 /**
 * Wrapper to contain the UIWebView and implement its delegate functions
 */
-
+#if !PLATFORM_TVOS
 @interface IOSWebViewWrapper : NSObject <UIWebViewDelegate>
+#else
+@interface IOSWebViewWrapper : NSObject
+#endif
 {
 	TSharedPtr<SIOSWebBrowserWidget> WebBrowserWidget;
 	bool bNeedsAddToView;
 }
+#if !PLATFORM_TVOS
 @property(strong) UIWebView* WebView;
+#endif
 @property(strong) NSURL* NextURL;
 @property(strong) NSString* NextContent;
 @property CGRect DesiredFrame;
