Fix TVOS compile error

diff --git a/Engine/Source/Runtime/WebBrowser/Private/IOS/IOSPlatformWebBrowser.h b/Engine/Source/Runtime/WebBrowser/Private/IOS/IOSPlatformWebBrowser.h
--- a/Engine/Source/Runtime/WebBrowser/Private/IOS/IOSPlatformWebBrowser.h
+++ b/Engine/Source/Runtime/WebBrowser/Private/IOS/IOSPlatformWebBrowser.h
@@ -9,7 +9,9 @@
 #include "IWebBrowserWindow.h"
 #include "Widgets/SWindow.h"
 #import <UIKit/UIKit.h>
+#if !PLATFORM_TVOS
 #import "WebKit/WebKit.h"
+#endif
 
 #include "Engine/Texture2D.h"
 #include "Materials/MaterialInstanceDynamic.h"
