Fix for missing Engine dependency in TVOS build

diff --git a/Engine/Source/Runtime/WebBrowser/WebBrowser.Build.cs b/Engine/Source/Runtime/WebBrowser/WebBrowser.Build.cs
--- a/Engine/Source/Runtime/WebBrowser/WebBrowser.Build.cs
+++ b/Engine/Source/Runtime/WebBrowser/WebBrowser.Build.cs
@@ -24,10 +24,11 @@ public class WebBrowser : ModuleRules
         );
 
         if (Target.Platform == UnrealTargetPlatform.Android ||
-            Target.Platform == UnrealTargetPlatform.IOS)
+            Target.Platform == UnrealTargetPlatform.IOS ||
+            Target.Platform == UnrealTargetPlatform.TVOS)
 		{
-			// We need these on Android for external texture support
-			PrivateDependencyModuleNames.AddRange(
+            // We need these on mobile for external texture support
+            PrivateDependencyModuleNames.AddRange(
 				new string[]
 				{
 					"WebBrowserTexture",
