UE-59012 //UE4/Main - Compile UE4Game Lumin - WARNING: Precompiled module 'PurchaseFlow' will not be usable without non-precompiled module 'WebBrowser'.
-Lumin is built with USE_ANDROID_JNI=0.  The Android implementation of the WebBrowser uses JNI, so it fails to compile.  This change wraps more JNI using code in #if USE_ANDROID_JNI and sets DUMMY_WEB_BROWSER if USE_ANDROID_JNI is false to avoid using the AndroidJavaWebBrowser.
-Previously WebBrowser had been set to not be precompiled, which avoided the compile breaks but created the warnings.

diff --git a/Engine/Source/Runtime/WebBrowser/Private/Android/AndroidJSScripting.cpp b/Engine/Source/Runtime/WebBrowser/Private/Android/AndroidJSScripting.cpp
--- a/Engine/Source/Runtime/WebBrowser/Private/Android/AndroidJSScripting.cpp
+++ b/Engine/Source/Runtime/WebBrowser/Private/Android/AndroidJSScripting.cpp
@@ -1,6 +1,9 @@
 // Copyright 1998-2018 Epic Games, Inc. All Rights Reserved.
 
 #include "AndroidJSScripting.h"
+
+#if USE_ANDROID_JNI
+
 #include "Android/AndroidApplication.h"
 #include "AndroidWebBrowserWindow.h"
 #include "AndroidWebBrowserDialog.h"
@@ -551,3 +554,5 @@ FAndroidJSScripting::FAndroidJSScripting(bool bJSBindingToLoweringEnabled)
 	: FWebJSScripting(bJSBindingToLoweringEnabled)
 {
 }
+
+#endif // USE_ANDROID_JNI
\ No newline at end of file
diff --git a/Engine/Source/Runtime/WebBrowser/Private/Android/AndroidJSScripting.h b/Engine/Source/Runtime/WebBrowser/Private/Android/AndroidJSScripting.h
--- a/Engine/Source/Runtime/WebBrowser/Private/Android/AndroidJSScripting.h
+++ b/Engine/Source/Runtime/WebBrowser/Private/Android/AndroidJSScripting.h
@@ -2,6 +2,9 @@
 #pragma once
 
 #include "CoreMinimal.h"
+
+#if USE_ANDROID_JNI
+
 #include "WebJSFunction.h"
 #include "WebJSScripting.h"
 
@@ -52,3 +55,5 @@ private:
 	/** Pointer to the Android Browser for this window. */
 	TWeakPtr<class FAndroidWebBrowserWindow> WindowPtr;
 };
+
+#endif // USE_ANDROID_JNI
\ No newline at end of file
diff --git a/Engine/Source/Runtime/WebBrowser/Private/Android/AndroidJSStructDeserializerBackend.cpp b/Engine/Source/Runtime/WebBrowser/Private/Android/AndroidJSStructDeserializerBackend.cpp
--- a/Engine/Source/Runtime/WebBrowser/Private/Android/AndroidJSStructDeserializerBackend.cpp
+++ b/Engine/Source/Runtime/WebBrowser/Private/Android/AndroidJSStructDeserializerBackend.cpp
@@ -1,6 +1,9 @@
 // Copyright 1998-2018 Epic Games, Inc. All Rights Reserved.
 
 #include "AndroidJSStructDeserializerBackend.h"
+
+#if USE_ANDROID_JNI
+
 #include "AndroidJSScripting.h"
 #include "UObject/UnrealType.h"
 #include "Templates/Casts.h"
@@ -99,3 +102,5 @@ FAndroidJSStructDeserializerBackend::FAndroidJSStructDeserializerBackend(FAndroi
 	auto Convert = StringCast<UCS2CHAR>(*JsonString);
 	JsonData.Append((uint8*)Convert.Get(), JsonString.Len() * sizeof(UCS2CHAR));
 }
+
+#endif // USE_ANDROID_JNI
\ No newline at end of file
diff --git a/Engine/Source/Runtime/WebBrowser/Private/Android/AndroidJSStructDeserializerBackend.h b/Engine/Source/Runtime/WebBrowser/Private/Android/AndroidJSStructDeserializerBackend.h
--- a/Engine/Source/Runtime/WebBrowser/Private/Android/AndroidJSStructDeserializerBackend.h
+++ b/Engine/Source/Runtime/WebBrowser/Private/Android/AndroidJSStructDeserializerBackend.h
@@ -2,6 +2,8 @@
 
 #pragma once
 
+#if USE_ANDROID_JNI
+
 #include "CoreMinimal.h"
 #include "AndroidJSScripting.h"
 #include "Backends/JsonStructDeserializerBackend.h"
@@ -20,3 +22,5 @@ private:
 	TArray<uint8> JsonData;
 	FMemoryReader Reader;
 };
+
+#endif // USE_ANDROID_JNI
\ No newline at end of file
diff --git a/Engine/Source/Runtime/WebBrowser/Private/Android/AndroidJSStructSerializerBackend.cpp b/Engine/Source/Runtime/WebBrowser/Private/Android/AndroidJSStructSerializerBackend.cpp
--- a/Engine/Source/Runtime/WebBrowser/Private/Android/AndroidJSStructSerializerBackend.cpp
+++ b/Engine/Source/Runtime/WebBrowser/Private/Android/AndroidJSStructSerializerBackend.cpp
@@ -1,6 +1,9 @@
 // Copyright 1998-2018 Epic Games, Inc. All Rights Reserved.
 
 #include "AndroidJSStructSerializerBackend.h"
+
+#if USE_ANDROID_JNI
+
 #include "AndroidJSScripting.h"
 #include "UObject/UnrealType.h"
 #include "UObject/PropertyPortFlags.h"
@@ -55,3 +58,5 @@ FAndroidJSStructSerializerBackend::FAndroidJSStructSerializerBackend(TSharedRef<
 	, FJsonStructSerializerBackend(Writer)
 {
 }
+
+#endif // USE_ANDROID_JNI
\ No newline at end of file
diff --git a/Engine/Source/Runtime/WebBrowser/Private/Android/AndroidJSStructSerializerBackend.h b/Engine/Source/Runtime/WebBrowser/Private/Android/AndroidJSStructSerializerBackend.h
--- a/Engine/Source/Runtime/WebBrowser/Private/Android/AndroidJSStructSerializerBackend.h
+++ b/Engine/Source/Runtime/WebBrowser/Private/Android/AndroidJSStructSerializerBackend.h
@@ -2,6 +2,8 @@
 
 #pragma once
 
+#if USE_ANDROID_JNI
+
 #include "CoreMinimal.h"
 #include "AndroidJSScripting.h"
 #include "Backends/JsonStructSerializerBackend.h"
@@ -38,3 +40,5 @@ private:
 	TArray<uint8> ReturnBuffer;
 	FMemoryWriter Writer;
 };
+
+#endif // USE_ANDROID_JNI
\ No newline at end of file
diff --git a/Engine/Source/Runtime/WebBrowser/Private/Android/AndroidJavaWebBrowser.cpp b/Engine/Source/Runtime/WebBrowser/Private/Android/AndroidJavaWebBrowser.cpp
--- a/Engine/Source/Runtime/WebBrowser/Private/Android/AndroidJavaWebBrowser.cpp
+++ b/Engine/Source/Runtime/WebBrowser/Private/Android/AndroidJavaWebBrowser.cpp
@@ -1,6 +1,9 @@
 // Copyright 1998-2018 Epic Games, Inc. All Rights Reserved.
 
 #include "AndroidJavaWebBrowser.h"
+
+#if USE_ANDROID_JNI
+
 #include "Android/AndroidApplication.h"
 
 #if UE_BUILD_SHIPPING
@@ -217,4 +220,5 @@ void FJavaAndroidWebBrowser::SetAndroid3DBrowser(bool InIsAndroid3DBrowser)
 	CallMethod<void>(SetAndroid3DBrowserMethod, InIsAndroid3DBrowser);
 }
 
+#endif // USE_ANDROID_JNI
 
diff --git a/Engine/Source/Runtime/WebBrowser/Private/Android/AndroidJavaWebBrowser.h b/Engine/Source/Runtime/WebBrowser/Private/Android/AndroidJavaWebBrowser.h
--- a/Engine/Source/Runtime/WebBrowser/Private/Android/AndroidJavaWebBrowser.h
+++ b/Engine/Source/Runtime/WebBrowser/Private/Android/AndroidJavaWebBrowser.h
@@ -2,6 +2,8 @@
 
 #pragma once
 
+#if USE_ANDROID_JNI
+
 #include "CoreMinimal.h"
 #include "Android/AndroidJava.h"
 #include "RHI.h"
@@ -78,3 +80,5 @@ public:
 	}
 
 };
+
+#endif // USE_ANDROID_JNI
\ No newline at end of file
diff --git a/Engine/Source/Runtime/WebBrowser/Private/Android/AndroidWebBrowserDialog.cpp b/Engine/Source/Runtime/WebBrowser/Private/Android/AndroidWebBrowserDialog.cpp
--- a/Engine/Source/Runtime/WebBrowser/Private/Android/AndroidWebBrowserDialog.cpp
+++ b/Engine/Source/Runtime/WebBrowser/Private/Android/AndroidWebBrowserDialog.cpp
@@ -1,6 +1,9 @@
 // Copyright 1998-2018 Epic Games, Inc. All Rights Reserved.
 
 #include "AndroidWebBrowserDialog.h"
+
+#if USE_ANDROID_JNI
+
 #include "Android/AndroidApplication.h"
 #include "Android/AndroidJava.h"
 
@@ -55,3 +58,5 @@ void FAndroidWebBrowserDialog::Continue(bool Success, const FText& UserResponse)
 	}
 	JEnv->CallVoidMethod(Callback, MethodId, JUserResponse);
 }
+
+#endif // USE_ANDROID_JNI
\ No newline at end of file
diff --git a/Engine/Source/Runtime/WebBrowser/Private/Android/AndroidWebBrowserDialog.h b/Engine/Source/Runtime/WebBrowser/Private/Android/AndroidWebBrowserDialog.h
--- a/Engine/Source/Runtime/WebBrowser/Private/Android/AndroidWebBrowserDialog.h
+++ b/Engine/Source/Runtime/WebBrowser/Private/Android/AndroidWebBrowserDialog.h
@@ -2,7 +2,10 @@
 
 #pragma once
 
+#if USE_ANDROID_JNI
+
 #include "IWebBrowserDialog.h"
+
 #include <jni.h>
 
 class SAndroidWebBrowserWidget;
@@ -58,3 +61,5 @@ private:
 };
 
 typedef FAndroidWebBrowserDialog FWebBrowserDialog;
+
+#endif // USE_ANDROID_JNI
\ No newline at end of file
diff --git a/Engine/Source/Runtime/WebBrowser/Private/Android/AndroidWebBrowserWidget.cpp b/Engine/Source/Runtime/WebBrowser/Private/Android/AndroidWebBrowserWidget.cpp
--- a/Engine/Source/Runtime/WebBrowser/Private/Android/AndroidWebBrowserWidget.cpp
+++ b/Engine/Source/Runtime/WebBrowser/Private/Android/AndroidWebBrowserWidget.cpp
@@ -1,6 +1,9 @@
 // Copyright 1998-2018 Epic Games, Inc. All Rights Reserved.
 
 #include "AndroidWebBrowserWidget.h"
+
+#if USE_ANDROID_JNI
+
 #include "AndroidWebBrowserWindow.h"
 #include "AndroidWebBrowserDialog.h"
 #include "AndroidJSScripting.h"
@@ -736,3 +739,5 @@ JNI_METHOD void Java_com_epicgames_ue4_WebViewControl_00024ChromeClient_onReceiv
 		Widget->HandleReceivedTitle(Title);
 	}
 }
+
+#endif // USE_ANDROID_JNI
\ No newline at end of file
diff --git a/Engine/Source/Runtime/WebBrowser/Private/Android/AndroidWebBrowserWidget.h b/Engine/Source/Runtime/WebBrowser/Private/Android/AndroidWebBrowserWidget.h
--- a/Engine/Source/Runtime/WebBrowser/Private/Android/AndroidWebBrowserWidget.h
+++ b/Engine/Source/Runtime/WebBrowser/Private/Android/AndroidWebBrowserWidget.h
@@ -2,6 +2,8 @@
 
 #pragma once
 
+#if USE_ANDROID_JNI
+
 #include "CoreMinimal.h"
 #include "Widgets/SLeafWidget.h"
 #include "Widgets/DeclarativeSyntaxSupport.h"
@@ -113,3 +115,5 @@ private:
 	/** Texture sample object pool. */
 	FWebBrowserTextureSamplePool* TextureSamplePool;
 };
+
+#endif // USE_ANDROID_JNI
\ No newline at end of file
diff --git a/Engine/Source/Runtime/WebBrowser/Private/Android/AndroidWebBrowserWindow.cpp b/Engine/Source/Runtime/WebBrowser/Private/Android/AndroidWebBrowserWindow.cpp
--- a/Engine/Source/Runtime/WebBrowser/Private/Android/AndroidWebBrowserWindow.cpp
+++ b/Engine/Source/Runtime/WebBrowser/Private/Android/AndroidWebBrowserWindow.cpp
@@ -1,6 +1,9 @@
 // Copyright 1998-2018 Epic Games, Inc. All Rights Reserved.
 
 #include "AndroidWebBrowserWindow.h"
+
+#if USE_ANDROID_JNI
+
 #include "AndroidWebBrowserDialog.h"
 #include "AndroidWebBrowserWidget.h"
 #include "Android/AndroidApplication.h"
@@ -285,3 +288,5 @@ void FAndroidWebBrowserWindow::UnbindUObject(const FString& Name, UObject* Objec
 {
 	Scripting->UnbindUObject(Name, Object, bIsPermanent);
 }
+
+#endif // USE_ANDROID_JNI
\ No newline at end of file
diff --git a/Engine/Source/Runtime/WebBrowser/Private/Android/AndroidWebBrowserWindow.h b/Engine/Source/Runtime/WebBrowser/Private/Android/AndroidWebBrowserWindow.h
--- a/Engine/Source/Runtime/WebBrowser/Private/Android/AndroidWebBrowserWindow.h
+++ b/Engine/Source/Runtime/WebBrowser/Private/Android/AndroidWebBrowserWindow.h
@@ -2,6 +2,8 @@
 
 #pragma once
 
+#if USE_ANDROID_JNI
+
 #include "IWebBrowserWindow.h"
 #include "AndroidJSScripting.h"
 #include "Widgets/SWidget.h"
@@ -267,3 +269,5 @@ private:
 };
 
 typedef FAndroidWebBrowserWindow FWebBrowserWindow;
+
+#endif // USE_ANDROID_JNI
\ No newline at end of file
diff --git a/Engine/Source/Runtime/WebBrowser/Private/SWebBrowserView.cpp b/Engine/Source/Runtime/WebBrowser/Private/SWebBrowserView.cpp
--- a/Engine/Source/Runtime/WebBrowser/Private/SWebBrowserView.cpp
+++ b/Engine/Source/Runtime/WebBrowser/Private/SWebBrowserView.cpp
@@ -12,7 +12,7 @@
 #include "WebBrowserViewport.h"
 #include "IWebBrowserAdapter.h"
 
-#if PLATFORM_ANDROID
+#if PLATFORM_ANDROID && USE_ANDROID_JNI
 #	include "Android/AndroidWebBrowserWindow.h"
 #elif PLATFORM_IOS
 #	include "IOS/IOSPlatformWebBrowser.h"
diff --git a/Engine/Source/Runtime/WebBrowser/Private/WebBrowserSingleton.cpp b/Engine/Source/Runtime/WebBrowser/Private/WebBrowserSingleton.cpp
--- a/Engine/Source/Runtime/WebBrowser/Private/WebBrowserSingleton.cpp
+++ b/Engine/Source/Runtime/WebBrowser/Private/WebBrowserSingleton.cpp
@@ -42,7 +42,7 @@ THIRD_PARTY_INCLUDES_END
 #	endif
 #endif
 
-#if PLATFORM_ANDROID
+#if PLATFORM_ANDROID && USE_ANDROID_JNI
 #	include "Android/AndroidWebBrowserWindow.h"
 #elif PLATFORM_IOS
 #	include <IOS/IOSPlatformWebBrowser.h>
@@ -495,7 +495,7 @@ TSharedPtr<IWebBrowserWindow> FWebBrowserSingleton::CreateBrowserWindow(const FC
 			return NewBrowserWindow;
 		}
 	}
-#elif PLATFORM_ANDROID
+#elif PLATFORM_ANDROID && USE_ANDROID_JNI
 	// Create new window
 	TSharedPtr<FAndroidWebBrowserWindow> NewBrowserWindow = MakeShareable(new FAndroidWebBrowserWindow(
 		WindowSettings.InitialURL,
diff --git a/Engine/Source/Runtime/WebBrowser/WebBrowser.Build.cs b/Engine/Source/Runtime/WebBrowser/WebBrowser.Build.cs
--- a/Engine/Source/Runtime/WebBrowser/WebBrowser.Build.cs
+++ b/Engine/Source/Runtime/WebBrowser/WebBrowser.Build.cs
@@ -9,7 +9,7 @@ public class WebBrowser : ModuleRules
 	{
 		PublicIncludePaths.Add("Runtime/WebBrowser/Public");
 		PrivateIncludePaths.Add("Runtime/WebBrowser/Private");
-        PrivateDependencyModuleNames.AddRange(
+		PrivateDependencyModuleNames.AddRange(
 			new string[]
 			{
 				"Core",
@@ -20,10 +20,10 @@ public class WebBrowser : ModuleRules
 				"Slate",
 				"SlateCore",
 				"Serialization",
-            }
-        );
+			}
+		);
 
-        if (Target.Platform == UnrealTargetPlatform.Android ||
+		if (Target.Platform == UnrealTargetPlatform.Android ||
             Target.Platform == UnrealTargetPlatform.IOS ||
             Target.Platform == UnrealTargetPlatform.TVOS)
 		{
@@ -39,10 +39,6 @@ public class WebBrowser : ModuleRules
 			// We need this one on Android for URL decoding
 			PrivateDependencyModuleNames.Add("HTTP");
 		}
-		if (Target.Platform == UnrealTargetPlatform.Lumin)
-		{
-			PrecompileForTargets = PrecompileTargetsType.None;
-		}
 
 		if (Target.Platform == UnrealTargetPlatform.Win64
 		||  Target.Platform == UnrealTargetPlatform.Win32
