UE-59052 Can't use touch to interact with Web Browser on iPhoneX

#iOS
#4.20

from CL4077699 on Dev-Mobile

Fixed crash on some iOS 11 devices (addSubview before setting WebView's params)
Fixed unresponsive WebView on iPhone X (resize the parent with the same size as the child, the touch won't work if the parent is smaller)

diff --git a/Engine/Source/Runtime/WebBrowser/Private/IOS/IOSPlatformWebBrowser.cpp b/Engine/Source/Runtime/WebBrowser/Private/IOS/IOSPlatformWebBrowser.cpp
--- a/Engine/Source/Runtime/WebBrowser/Private/IOS/IOSPlatformWebBrowser.cpp
+++ b/Engine/Source/Runtime/WebBrowser/Private/IOS/IOSPlatformWebBrowser.cpp
@@ -314,7 +314,10 @@ supportsMetal : (bool)InSupportsMetal supportsMetalMRT : (bool)InSupportsMetalMR
 
 
 		WebView = [[WKWebView alloc]initWithFrame:CGRectMake(1, 1, 100, 100)];
+		[self.WebViewContainer addSubview : WebView];
 		WebView.navigationDelegate = self;
+		WebView.UIDelegate = self;
+
 		WebView.scrollView.bounces = NO;
 
 		if (InUseTransparency)
@@ -327,8 +330,6 @@ supportsMetal : (bool)InSupportsMetal supportsMetalMRT : (bool)InSupportsMetalMR
 			[self.WebView setOpaque : YES];
 		}
 
-		[self.WebViewContainer addSubview : WebView];
-
 		[self setWebViewVisible];
 	});
 #endif
@@ -341,9 +342,9 @@ supportsMetal : (bool)InSupportsMetal supportsMetalMRT : (bool)InSupportsMetalMR
 	dispatch_async(dispatch_get_main_queue(), ^
 	{
 		[self.WebViewContainer removeFromSuperview];
-	[self.WebView removeFromSuperview];
-	WebView = nil;
-	WebViewContainer = nil;
+		[self.WebView removeFromSuperview];
+		WebView = nil;
+		WebViewContainer = nil;
 	});
 #endif
 }
@@ -357,7 +358,8 @@ supportsMetal : (bool)InSupportsMetal supportsMetalMRT : (bool)InSupportsMetalMR
 	{
 		if (WebView != nil)
 		{
-			WebView.frame = self.DesiredFrame;
+			WebViewContainer.frame = self.DesiredFrame;
+			WebView.frame = WebViewContainer.bounds;
 			if (bNeedsAddToView)
 			{
 				bNeedsAddToView = false;
@@ -419,6 +421,7 @@ supportsMetal : (bool)InSupportsMetal supportsMetalMRT : (bool)InSupportsMetalMR
 	{
 		if (IsIOS3DBrowser != InIsIOS3DBrowser)
 		{
+			//default is 2D
 			IsIOS3DBrowser = InIsIOS3DBrowser;
 			[self setWebViewVisible];
 		}
