Copying //UE4/Dev-Main to //UE4/Main (Source: //Portal/Main/Engine @ 4247640)

diff --git a/Engine/Source/Runtime/Online/BuildPatchServices/Public/BuildPatchSettings.h b/Engine/Source/Runtime/Online/BuildPatchServices/Public/BuildPatchSettings.h
--- a/Engine/Source/Runtime/Online/BuildPatchServices/Public/BuildPatchSettings.h
+++ b/Engine/Source/Runtime/Online/BuildPatchServices/Public/BuildPatchSettings.h
@@ -2,12 +2,32 @@
 #pragma once
 
 #include "CoreMinimal.h"
-#include "Interfaces/IBuildManifest.h"
 #include "Misc/Variant.h"
+#include "Interfaces/IBuildManifest.h"
+#include "BuildPatchInstall.h"
 #include "BuildPatchVerify.h"
 
 namespace BuildPatchServices
 {
+	/**
+	 * Defines a list of all build patch services initialization settings, can be used to override default init behaviors.
+	 */
+	struct FBuildPatchServicesInitSettings
+	{
+	public:
+		/**
+		 * Default constructor. Initializes all members with default behavior values.
+		 */
+		FBuildPatchServicesInitSettings();
+
+		// The application settings directory.
+		FString ApplicationSettingsDir;
+		// The application project name.
+		FString ProjectName;
+		// The local machine config file name.
+		FString LocalMachineConfigFileName;
+	};
+
 	/**
 	 * Defines a list of all the options of an installation task.
 	 */
@@ -29,14 +49,16 @@ namespace BuildPatchServices
 		TArray<FString> CloudDirectories;
 		// The set of tags that describe what to be installed. Empty set means full installation.
 		TSet<FString> InstallTags;
+		// The mode for installation.
+		EInstallMode InstallMode;
 		// The mode for verification.
 		EVerifyMode VerifyMode;
 		// Whether the operation is a repair to an existing installation only.
 		bool bIsRepair;
-		// Whether the operation should only produce the necessary staged data, without performing the final install stage yet.
-		bool bStageOnly;
 		// Whether to run the prerequisite installer provided if it hasn't been ran before on this machine.
 		bool bRunRequiredPrereqs;
+		// Whether to allow this installation to run concurrently with any existing installations.
+		bool bAllowConcurrentExecution;
 
 		/**
 		 * Construct with install manifest, provides common defaults for other settings.
@@ -50,10 +72,11 @@ namespace BuildPatchServices
 			, ChunkDatabaseFiles()
 			, CloudDirectories()
 			, InstallTags()
+			, InstallMode(EInstallMode::NonDestructiveInstall)
 			, VerifyMode(EVerifyMode::ShaVerifyAllFiles)
 			, bIsRepair(false)
-			, bStageOnly(false)
 			, bRunRequiredPrereqs(true)
+			, bAllowConcurrentExecution(false)
 		{}
 
 		/**
@@ -68,10 +91,11 @@ namespace BuildPatchServices
 			, ChunkDatabaseFiles(MoveTemp(MoveFrom.ChunkDatabaseFiles))
 			, CloudDirectories(MoveTemp(MoveFrom.CloudDirectories))
 			, InstallTags(MoveTemp(MoveFrom.InstallTags))
-			, VerifyMode(MoveTemp(MoveFrom.VerifyMode))
-			, bIsRepair(MoveTemp(MoveFrom.bIsRepair))
-			, bStageOnly(MoveTemp(MoveFrom.bStageOnly))
-			, bRunRequiredPrereqs(MoveTemp(MoveFrom.bRunRequiredPrereqs))
+			, InstallMode(MoveFrom.InstallMode)
+			, VerifyMode(MoveFrom.VerifyMode)
+			, bIsRepair(MoveFrom.bIsRepair)
+			, bRunRequiredPrereqs(MoveFrom.bRunRequiredPrereqs)
+			, bAllowConcurrentExecution(MoveFrom.bAllowConcurrentExecution)
 		{}
 
 		/**
@@ -86,10 +110,11 @@ namespace BuildPatchServices
 			, ChunkDatabaseFiles(CopyFrom.ChunkDatabaseFiles)
 			, CloudDirectories(CopyFrom.CloudDirectories)
 			, InstallTags(CopyFrom.InstallTags)
+			, InstallMode(CopyFrom.InstallMode)
 			, VerifyMode(CopyFrom.VerifyMode)
 			, bIsRepair(CopyFrom.bIsRepair)
-			, bStageOnly(CopyFrom.bStageOnly)
 			, bRunRequiredPrereqs(CopyFrom.bRunRequiredPrereqs)
+			, bAllowConcurrentExecution(CopyFrom.bAllowConcurrentExecution)
 		{}
 	};
 
diff --git a/Engine/Source/Runtime/Online/BuildPatchServices/Public/BuildPatchState.h b/Engine/Source/Runtime/Online/BuildPatchServices/Public/BuildPatchState.h
--- a/Engine/Source/Runtime/Online/BuildPatchServices/Public/BuildPatchState.h
+++ b/Engine/Source/Runtime/Online/BuildPatchServices/Public/BuildPatchState.h
@@ -62,5 +62,6 @@ namespace BuildPatchServices
 	 * @param State - The error type value.
 	 * @return The display text associated with the progress step.
 	 */
+	DEPRECATED(4.21, "BuildPatchServices::StateToText(const EBuildPatchState& State) has been deprecated. It will no longer be supported by BuildPatchServices in the future.")
 	BUILDPATCHSERVICES_API const FText& StateToText(const EBuildPatchState& State);
 }
diff --git a/Engine/Source/Runtime/Online/BuildPatchServices/Public/BuildPatchVerify.h b/Engine/Source/Runtime/Online/BuildPatchServices/Public/BuildPatchVerify.h
--- a/Engine/Source/Runtime/Online/BuildPatchServices/Public/BuildPatchVerify.h
+++ b/Engine/Source/Runtime/Online/BuildPatchServices/Public/BuildPatchVerify.h
@@ -2,6 +2,7 @@
 #pragma once
 
 #include "CoreMinimal.h"
+#include "Misc/EnumRange.h"
 
 namespace BuildPatchServices
 {
@@ -22,4 +23,48 @@ namespace BuildPatchServices
 		// Checks just the existence and file size of only files touched by the install/patch process.
 		FileSizeCheckTouchedFiles
 	};
+
+	/**
+	 * An enum defining the possible causes for a verification failure.
+	 */
+	enum class EVerifyError : uint32
+	{
+		// The file was not found.
+		FileMissing,
+
+		// The file failed to open.
+		OpenFileFailed,
+
+		// The file failed its hash check.
+		HashCheckFailed,
+
+		// The file was not the expected size.
+		FileSizeFailed
+	};
+
+	/**
+	 * Returns the string representation of the EVerifyMode value. Used for analytics and logging only.
+	 * @param  VerifyMode     The value.
+	 * @return the enum's string representation.
+	 */
+	inline const FString& EnumToString(const EVerifyMode& VerifyMode)
+	{
+		// Const enum strings, special case no error.
+		static const FString ShaVerifyAllFiles(TEXT("EVerifyMode::ShaVerifyAllFiles"));
+		static const FString ShaVerifyTouchedFiles(TEXT("EVerifyMode::ShaVerifyTouchedFiles"));
+		static const FString FileSizeCheckAllFiles(TEXT("EVerifyMode::FileSizeCheckAllFiles"));
+		static const FString FileSizeCheckTouchedFiles(TEXT("EVerifyMode::FileSizeCheckTouchedFiles"));
+		static const FString InvalidOrMax(TEXT("InvalidOrMax"));
+
+		switch (VerifyMode)
+		{
+			case EVerifyMode::ShaVerifyAllFiles: return ShaVerifyAllFiles;
+			case EVerifyMode::ShaVerifyTouchedFiles: return ShaVerifyTouchedFiles;
+			case EVerifyMode::FileSizeCheckAllFiles: return FileSizeCheckAllFiles;
+			case EVerifyMode::FileSizeCheckTouchedFiles: return FileSizeCheckTouchedFiles;
+			default: return InvalidOrMax;
+		}
+	}
 }
+
+ENUM_RANGE_BY_FIRST_AND_LAST(BuildPatchServices::EVerifyError, BuildPatchServices::EVerifyError::FileMissing, BuildPatchServices::EVerifyError::FileSizeFailed)
\ No newline at end of file
diff --git a/Engine/Source/Runtime/Online/BuildPatchServices/Public/Interfaces/IBuildInstaller.h b/Engine/Source/Runtime/Online/BuildPatchServices/Public/Interfaces/IBuildInstaller.h
--- a/Engine/Source/Runtime/Online/BuildPatchServices/Public/Interfaces/IBuildInstaller.h
+++ b/Engine/Source/Runtime/Online/BuildPatchServices/Public/Interfaces/IBuildInstaller.h
@@ -12,8 +12,9 @@
 
 class IBuildInstaller;
 
-typedef TSharedPtr< class IBuildInstaller, ESPMode::ThreadSafe > IBuildInstallerPtr;
-typedef TSharedRef< class IBuildInstaller, ESPMode::ThreadSafe > IBuildInstallerRef;
+typedef TSharedPtr<IBuildInstaller, ESPMode::ThreadSafe> IBuildInstallerPtr;
+typedef TSharedRef<IBuildInstaller, ESPMode::ThreadSafe> IBuildInstallerRef;
+typedef TWeakPtr<IBuildInstaller, ESPMode::ThreadSafe> IBuildInstallerWeakPtr;
 
 /**
  * Declares the error type enum for use with the error system
@@ -121,18 +122,26 @@ struct FBuildInstallStats
 		, NumChunksDownloaded(0)
 		, NumChunksRecycled(0)
 		, NumChunksReadFromChunkDbs(0)
-		, NumChunksCacheBooted(0)
-		, NumDriveCacheChunkLoads(0)
 		, NumFailedDownloads(0)
 		, NumBadDownloads(0)
 		, NumAbortedDownloads(0)
 		, NumRecycleFailures(0)
-		, NumDriveCacheLoadFailures(0)
+		, NumChunksStoreBooted(0)
+		, NumDriveStoreChunkLoads(0)
+		, NumDriveStoreLoadFailures(0)
 		, NumChunkDbChunksFailed(0)
 		, TotalDownloadedData(0)
 		, AverageDownloadSpeed(0.0)
+		, PeakDownloadSpeed(0.0)
 		, FinalDownloadSpeed(-1.0)
 		, TheoreticalDownloadTime(0.0f)
+		, TotalReadData(0)
+		, AverageDiskReadSpeed(0.0)
+		, PeakDiskReadSpeed(0.0)
+		, TotalWrittenData(0)
+		, AverageDiskWriteSpeed(0.0)
+		, PeakDiskWriteSpeed(0.0)
+		, NumFilesConstructed(0)
 		, InitializeTime(0.0f)
 		, ConstructTime(0.0f)
 		, MoveFromStageTime(0.0f)
@@ -157,6 +166,11 @@ struct FBuildInstallStats
 		, OkDownloadHealthTime(0.0f)
 		, PoorDownloadHealthTime(0.0f)
 		, DisconnectedDownloadHealthTime(0.0f)
+		, AverageMemoryStoreUse(0.0f)
+		, PeakMemoryStoreUse(0)
+		, AverageMemoryStoreRetained(0.0f)
+		, PeakMemoryStoreRetained(0)
+		, MemoryStoreSize(0)
 	{}
 
 	// The name of the app being installed.
@@ -187,10 +201,6 @@ struct FBuildInstallStats
 	uint32 NumChunksRecycled;
 	// The number of chunks successfully read from chunkdbs.
 	uint32 NumChunksReadFromChunkDbs;
-	// The number of chunks that had to be booted from the cache.
-	uint32 NumChunksCacheBooted;
-	// The number of chunks that had to be loaded from the drive cache.
-	uint32 NumDriveCacheChunkLoads;
 	// The number of chunks we did not successfully receive.
 	uint32 NumFailedDownloads;
 	// The number of chunks we received but were determined bad data.
@@ -199,18 +209,38 @@ struct FBuildInstallStats
 	uint32 NumAbortedDownloads;
 	// The number of chunks that failed to be recycled from existing build.
 	uint32 NumRecycleFailures;
-	// The number of chunks that failed to load from the drive cache.
-	uint32 NumDriveCacheLoadFailures;
+	// The number of chunks that had to be booted from the memory stores.
+	uint32 NumChunksStoreBooted;
+	// The number of chunks that had to be loaded from the drive store.
+	uint32 NumDriveStoreChunkLoads;
+	// The number of chunks that failed to load from the drive store.
+	uint32 NumDriveStoreLoadFailures;
 	// The number of chunks that were not successfully loaded from provided chunkdbs.
 	uint32 NumChunkDbChunksFailed;
 	// The total number of bytes downloaded.
-	int64 TotalDownloadedData;
+	uint64 TotalDownloadedData;
 	// The average chunk download speed.
 	double AverageDownloadSpeed;
+	// The peak chunk download speed.
+	double PeakDownloadSpeed;
 	// The download speed registered at the end of the installation.
 	double FinalDownloadSpeed;
 	// The theoretical download time (data/speed).
 	float TheoreticalDownloadTime;
+	// The total number of bytes read to disk.
+	uint64 TotalReadData;
+	// The average disk read speed.
+	double AverageDiskReadSpeed;
+	// The peak disk read speed.
+	double PeakDiskReadSpeed;
+	// The total number of bytes written to disk.
+	uint64 TotalWrittenData;
+	// The average disk write speed.
+	double AverageDiskWriteSpeed;
+	// The peak disk write speed.
+	double PeakDiskWriteSpeed;
+	// The total number of files constructed.
+	uint32 NumFilesConstructed;
 	// The time spent during the initialization stage.
 	float InitializeTime;
 	// The time spent during the construction stage.
@@ -259,10 +289,20 @@ struct FBuildInstallStats
 	float PoorDownloadHealthTime;
 	// The amount of time that was spent with Disconnected download health.
 	float DisconnectedDownloadHealthTime;
+	// AverageMemoryStoreUse
+	float AverageMemoryStoreUse;
+	// PeakMemoryStoreUse
+	uint32 PeakMemoryStoreUse;
+	// AverageMemoryStoreRetained
+	float AverageMemoryStoreRetained;
+	// PeakMemoryStoreRetained
+	uint32 PeakMemoryStoreRetained;
+	// MemoryStoreSize
+	uint32 MemoryStoreSize;
 };
 
 /**
- * Interface to a Build Manifest.
+ * Interface to a Build Installer, exposes installation control, progress, and state information.
  */
 class IBuildInstaller
 {
@@ -296,6 +336,12 @@ public:
 	 */
 	virtual bool IsResumable() const = 0;
 
+	/**
+	 * Get whether the install is performing an update of an existing install.
+	 * @return	true if installation is an update
+	 */
+	virtual bool IsUpdate() const = 0;
+
 	/**
 	 * Get whether the install failed. Only valid if complete.
 	 * @return	true if installation was a failure
@@ -308,20 +354,10 @@ public:
 	 */
 	virtual EBuildPatchInstallError GetErrorType() const = 0;
 
-	/**
-	 * This is deprecated and shouldn't be used anymore [6/4/2014 justin.sargent]
-	 *
-	 * Get the percentage complete text for the current process
-	 * @return	percentage complete text progress
-	 */
+	DEPRECATED(4.21, "GetPercentageText has been deprecated.  It will no longer be supported in the future.")
 	virtual FText GetPercentageText() const = 0;
 
-	/**
-	 * This is deprecated and shouldn't be used anymore [6/4/2014 justin.sargent]
-	 *
-	 * Get the download speed text for the current process
-	 * @return	download speed text progress
-	 */
+	DEPRECATED(4.21, "GetDownloadSpeedText has been deprecated.  It will no longer be supported in the future.")
 	virtual FText GetDownloadSpeedText() const = 0;
 
 	/**
@@ -330,15 +366,18 @@ public:
 	 */
 	virtual double GetDownloadSpeed() const = 0;
 
+	DEPRECATED(4.21, "GetInitialDownloadSize has been deprecated.  Please use GetTotalDownloadRequired instead.")
+	virtual int64 GetInitialDownloadSize() const { return GetTotalDownloadRequired(); }
+
 	/**
-	 * Get the initial download size
-	 * @return	the initial download size
+	 * Get the total download bytes required to complete
+	 * @return	the total required bytes
 	 */
-	virtual int64 GetInitialDownloadSize() const = 0;
+	virtual int64 GetTotalDownloadRequired() const = 0;
 
 	/**
-	 * Get the total currently downloaded
-	 * @return	the total currently downloaded
+	 * Get the total bytes downloaded
+	 * @return	the total bytes downloaded
 	 */
 	virtual int64 GetTotalDownloaded() const = 0;
 
@@ -348,13 +387,7 @@ public:
 	 */
 	virtual BuildPatchServices::EBuildPatchState GetState() const = 0;
 
-	/**
-	 * This is deprecated and redundant [2016/06/12 leigh.swift]
-	 * Instead, you can BuildPatchServices::StateToText(Installer->GetState()).
-	 *
-	 * Get the text for status of the install process.
-	 * @return Status of the install process.
-	 */
+	DEPRECATED(4.21, "GetStatusText has been deprecated.  It will no longer be supported in the future.")
 	virtual FText GetStatusText() const = 0;
 
 	/**
diff --git a/Engine/Source/Runtime/Online/BuildPatchServices/Public/Interfaces/IBuildManifest.h b/Engine/Source/Runtime/Online/BuildPatchServices/Public/Interfaces/IBuildManifest.h
--- a/Engine/Source/Runtime/Online/BuildPatchServices/Public/Interfaces/IBuildManifest.h
+++ b/Engine/Source/Runtime/Online/BuildPatchServices/Public/Interfaces/IBuildManifest.h
@@ -158,10 +158,17 @@ public:
 
 	/**
 	 * Get the list of files in this build. Filenames are all relative to an install directory
-	 * @return		an array containing build files
+	 * @return an array containing build files
 	 */
 	virtual TArray<FString> GetBuildFileList() const = 0;
 
+	/**
+	 * Get the list of files in this build which match a tag from a given set. Filenames are all relative to an install directory
+	 * @param		Tags	The set of tags to query
+	 * @return an array containing build files
+	 */
+	virtual TArray<FString> GetBuildFileList(const TSet<FString>& Tags) const = 0;
+
 	/**
 	 * Get the list of install tags in this manifest
 	 * @param Tags			OUT		Receives the tags referenced.
diff --git a/Engine/Source/Runtime/Online/BuildPatchServices/Public/Interfaces/IBuildPatchServicesModule.h b/Engine/Source/Runtime/Online/BuildPatchServices/Public/Interfaces/IBuildPatchServicesModule.h
--- a/Engine/Source/Runtime/Online/BuildPatchServices/Public/Interfaces/IBuildPatchServicesModule.h
+++ b/Engine/Source/Runtime/Online/BuildPatchServices/Public/Interfaces/IBuildPatchServicesModule.h
@@ -2,8 +2,9 @@
 #pragma once
 
 #include "CoreMinimal.h"
-#include "Interfaces/IBuildInstaller.h"
 #include "Modules/ModuleInterface.h"
+#include "Interfaces/IBuildInstaller.h"
+#include "Interfaces/IBuildStatistics.h"
 #include "BuildPatchSettings.h"
 
 class FHttpServiceTracker;
@@ -33,22 +34,29 @@ public:
 	/**
 	 * Virtual destructor.
 	 */
-	virtual ~IBuildPatchServicesModule( ) { }
-	
+	virtual ~IBuildPatchServicesModule() { }
+
+	/**
+	 * Factory providing construction of a build statistics class.
+	 * @param Installer     The installer to create a build statistics for.
+	 * @return an instance of an IBuildStatistics implementation. 
+	 */
+	virtual BuildPatchServices::IBuildStatisticsRef CreateBuildStatistics(const IBuildInstallerRef& Installer) const = 0;
+
 	/**
 	 * Loads a Build Manifest from file and returns the interface
 	 * @param Filename		The file to load from
 	 * @return		a shared pointer to the manifest, which will be invalid if load failed.
 	 */
 	virtual IBuildManifestPtr LoadManifestFromFile( const FString& Filename ) = 0;
-	
+
 	/**
 	 * Constructs a Build Manifest from a data
 	 * @param ManifestData		The data received from a web api
 	 * @return		a shared pointer to the manifest, which will be invalid if creation failed.
 	 */
 	virtual IBuildManifestPtr MakeManifestFromData( const TArray<uint8>& ManifestData ) = 0;
-	
+
 	/**
 	 * Saves a Build Manifest to file
 	 * @param Filename		The file to save to
@@ -58,6 +66,12 @@ public:
 	 */
 	virtual bool SaveManifestToFile(const FString& Filename, IBuildManifestRef Manifest, bool bUseBinary = true) = 0;
 
+	/**
+	 * Gets an array of prerequisite identifiers that are registered as installed on this system.
+	 * @return a set containing all installed prerequisite identifiers.
+	 */
+	virtual TSet<FString> GetInstalledPrereqIds() const = 0;
+
 	/**
 	 * Starts an installer thread for the provided manifests
 	 * NB: THIS FUNCTION WILL EVENTUALLY BE DEPRECATED.
@@ -95,6 +109,12 @@ public:
 	 */
 	virtual IBuildInstallerRef StartBuildInstall(BuildPatchServices::FInstallerConfiguration Configuration, FBuildPatchBoolManifestDelegate OnCompleteDelegate) = 0;
 
+	/**
+	 * Gets a list of currently active installers
+	 * @return all installers that are currently active.
+	 */
+	virtual const TArray<IBuildInstallerRef>& GetInstallers() const = 0;
+
 	/**
 	 * Sets the directory used for staging intermediate files.
 	 * @param StagingDir	The staging directory
@@ -185,12 +205,15 @@ public:
 	/**
 	 * Packages data referenced by a manifest file into chunkdb files, supporting a maximum filesize per chunkdb.
 	 * @param ManifestFilePath      A full file path to the manifest to enumerate chunks from.
+	 * @param PrevManifestFilePath  A full file path to a manifest describing a previous build, which will filter out saved chunks for patch only chunkdbs.
+	 * @param TagSetArray           Optional list of tagsets to split chunkdb files on. Empty array will include all data as normal.
 	 * @param OutputFile            A full file path to the chunkdb file to save. Extension of .chunkdb will be added if not present.
 	 * @param CloudDir              Cloud directory where chunks to be packaged can be found.
 	 * @param MaxOutputFileSize     The maximum desired size for each chunkdb file.
+	 * @param ResultDataFilePath    A full file path to use when saving the json output data.
 	 * @return true if successful.
 	 */
-	virtual bool PackageChunkData(const FString& ManifestFilePath, const FString& OutputFile, const FString& CloudDir, uint64 MaxOutputFileSize) = 0;
+	virtual bool PackageChunkData(const FString& ManifestFilePath, const FString& PrevManifestFilePath, const TArray<TSet<FString>>& TagSetArray, const FString& OutputFile, const FString& CloudDir, uint64 MaxOutputFileSize, const FString& ResultDataFilePath) = 0;
 
 	/**
 	 * Takes two manifests as input, in order to merge together producing a new manifest containing all files.
@@ -217,14 +240,9 @@ public:
 	 */
 	virtual bool DiffManifests(const FString& ManifestFilePathA, const TSet<FString>& TagSetA, const FString& ManifestFilePathB, const TSet<FString>& TagSetB, const FString& OutputFilePath) = 0;
 
-	/**
-	 * Deprecated function, use MakeManifestFromData instead.
-	 */
+	DEPRECATED(4.21, "MakeManifestFromJSON(const FString& ManifestJSON) has been deprecated.  Please use MakeManifestFromData(const TArray<uint8>& ManifestData) instead.")
 	virtual IBuildManifestPtr MakeManifestFromJSON(const FString& ManifestJSON) = 0;
 
 	DEPRECATED(4.16, "Please use EnumeratePatchData instead.")
-	virtual bool EnumerateManifestData(const FString& ManifestFilePath, const FString& OutputFile, bool bIncludeSizes)
-	{
-		return EnumeratePatchData(ManifestFilePath, OutputFile, bIncludeSizes);
-	}
+	virtual bool EnumerateManifestData(const FString& ManifestFilePath, const FString& OutputFile, bool bIncludeSizes) { return EnumeratePatchData(ManifestFilePath, OutputFile, bIncludeSizes); }
 };
diff --git a/Engine/Source/Runtime/Online/HTTP/Private/Apple/AppleHTTP.cpp b/Engine/Source/Runtime/Online/HTTP/Private/Apple/AppleHTTP.cpp
--- a/Engine/Source/Runtime/Online/HTTP/Private/Apple/AppleHTTP.cpp
+++ b/Engine/Source/Runtime/Online/HTTP/Private/Apple/AppleHTTP.cpp
@@ -27,6 +27,13 @@ FAppleHttpRequest::FAppleHttpRequest()
 
 	// Disable cache to mimic WinInet behavior
 	Request.cachePolicy = NSURLRequestReloadIgnoringLocalCacheData;
+
+	// Add default headers
+	const TMap<FString, FString>& DefaultHeaders = FHttpModule::Get().GetDefaultHeaders();
+	for (TMap<FString, FString>::TConstIterator It(DefaultHeaders); It; ++It)
+	{
+		SetHeader(It.Key(), It.Value());
+	}
 }
 
 
@@ -301,6 +308,10 @@ void FAppleHttpRequest::FinishedRequest()
 		UE_LOG(LogHttp, Verbose, TEXT("Request failed"));
 		FString URL([[Request URL] absoluteString]);
 		CompletionStatus = EHttpRequestStatus::Failed;
+		if (Response.IsValid() && [Response->ResponseWrapper bIsHostConnectionFailure])
+		{
+			CompletionStatus = EHttpRequestStatus::Failed_ConnectionError;
+		}
 
 		Response = nullptr;
 		OnProcessRequestComplete().ExecuteIfBound(SharedThis(this), nullptr, false);
@@ -392,6 +403,7 @@ float FAppleHttpRequest::GetElapsedTime() const
 @synthesize Response;
 @synthesize bIsReady;
 @synthesize bHadError;
+@synthesize bIsHostConnectionFailure;
 @synthesize BytesWritten;
 
 
@@ -400,6 +412,8 @@ float FAppleHttpRequest::GetElapsedTime() const
 	UE_LOG(LogHttp, Verbose, TEXT("-(FHttpResponseAppleWrapper*) init"));
 	self = [super init];
 	bIsReady = false;
+	bHadError = false;
+	bIsHostConnectionFailure = false;
 	
 	return self;
 }
@@ -443,6 +457,14 @@ float FAppleHttpRequest::GetElapsedTime() const
 		*FString([error localizedDescription]),
 		*FString([[error userInfo] objectForKey:NSURLErrorFailingURLStringErrorKey]),
 		self);
+	// Determine if the specific error was failing to connect to the host.
+	switch ([error code])
+	{
+		case NSURLErrorCannotFindHost:
+		case NSURLErrorCannotConnectToHost:
+		case NSURLErrorDNSLookupFailed:
+			self.bIsHostConnectionFailure = YES;
+	}
 	// Log more details if verbose logging is enabled and this is an SSL error
 	if (UE_LOG_ACTIVE(LogHttp, Verbose))
 	{
diff --git a/Engine/Source/Runtime/Online/HTTP/Private/Apple/AppleHTTP.h b/Engine/Source/Runtime/Online/HTTP/Private/Apple/AppleHTTP.h
--- a/Engine/Source/Runtime/Online/HTTP/Private/Apple/AppleHTTP.h
+++ b/Engine/Source/Runtime/Online/HTTP/Private/Apple/AppleHTTP.h
@@ -124,6 +124,8 @@ private:
 @property BOOL bIsReady;
 /** When the response is complete, indicates whether the response was received without error. */
 @property BOOL bHadError;
+/** When the response is complete, indicates whether the response failed with an error specific to connecting to the host. */
+@property BOOL bIsHostConnectionFailure;
 /** The total number of bytes written out during the request/response */
 @property int32 BytesWritten;
 
diff --git a/Engine/Source/Runtime/Online/HTTP/Private/HTML5/HTML5HTTP.cpp b/Engine/Source/Runtime/Online/HTTP/Private/HTML5/HTML5HTTP.cpp
--- a/Engine/Source/Runtime/Online/HTTP/Private/HTML5/HTML5HTTP.cpp
+++ b/Engine/Source/Runtime/Online/HTTP/Private/HTML5/HTML5HTTP.cpp
@@ -20,6 +20,12 @@ FHTML5HttpRequest::FHTML5HttpRequest()
 	,	ElapsedTime(0.0f)
 {
 	UE_LOG(LogHttp, Verbose, TEXT("FHTML5HttpRequest::FHTML5HttpRequest()"));
+	// Add default headers
+	const TMap<FString, FString>& DefaultHeaders = FHttpModule::Get().GetDefaultHeaders();
+	for (TMap<FString, FString>::TConstIterator It(DefaultHeaders); It; ++It)
+	{
+		SetHeader(It.Key(), It.Value());
+	}
 }
 
 
diff --git a/Engine/Source/Runtime/Online/HTTP/Private/HttpManager.cpp b/Engine/Source/Runtime/Online/HTTP/Private/HttpManager.cpp
--- a/Engine/Source/Runtime/Online/HTTP/Private/HttpManager.cpp
+++ b/Engine/Source/Runtime/Online/HTTP/Private/HttpManager.cpp
@@ -5,6 +5,7 @@
 #include "HAL/PlatformProcess.h"
 #include "Misc/ScopeLock.h"
 #include "Http.h"
+#include "Misc/Guid.h"
 
 #include "HttpThread.h"
 #include "Misc/ConfigCacheIni.h"
@@ -17,6 +18,8 @@ FCriticalSection FHttpManager::RequestLock;
 
 FHttpManager::FHttpManager()
 	:	FTickerObjectBase(0.0f)
+	,	Thread(nullptr)
+	,	CorrelationIdMethod(FHttpManager::GetDefaultCorrelationIdMethod())
 	,	DeferredDestroyDelay(10.0f)
 {
 }
@@ -36,6 +39,23 @@ void FHttpManager::Initialize()
 	Thread->StartThread();
 }
 
+void FHttpManager::SetCorrelationIdMethod(TFunction<FString()> InCorrelationIdMethod)
+{
+	check(InCorrelationIdMethod);
+	CorrelationIdMethod = MoveTemp(InCorrelationIdMethod);
+}
+
+FString FHttpManager::CreateCorrelationId() const
+{
+	return CorrelationIdMethod();
+}
+
+/*static*/
+TFunction<FString()> FHttpManager::GetDefaultCorrelationIdMethod()
+{
+	return []{ return FGuid::NewGuid().ToString(); };
+}
+
 FHttpThread* FHttpManager::CreateHttpThread()
 {
 	return new FHttpThread();
@@ -199,4 +219,4 @@ void FHttpManager::DumpRequests(FOutputDevice& Ar) const
 bool FHttpManager::SupportsDynamicProxy() const
 {
 	return false;
-}
\ No newline at end of file
+}
diff --git a/Engine/Source/Runtime/Online/HTTP/Private/Windows/HttpWinInet.cpp b/Engine/Source/Runtime/Online/HTTP/Private/Windows/HttpWinInet.cpp
--- a/Engine/Source/Runtime/Online/HTTP/Private/Windows/HttpWinInet.cpp
+++ b/Engine/Source/Runtime/Online/HTTP/Private/Windows/HttpWinInet.cpp
@@ -355,7 +355,12 @@ FHttpRequestWinInet::FHttpRequestWinInet()
 ,	ElapsedTime(0.0f)
 ,	bDebugVerbose(false)
 {
-
+	// Add default headers
+	const TMap<FString, FString>& DefaultHeaders = FHttpModule::Get().GetDefaultHeaders();
+	for (TMap<FString, FString>::TConstIterator It(DefaultHeaders); It; ++It)
+	{
+		SetHeader(It.Key(), It.Value());
+	}
 }
 
 FHttpRequestWinInet::~FHttpRequestWinInet()
diff --git a/Engine/Source/Runtime/Online/HTTP/Public/HttpManager.h b/Engine/Source/Runtime/Online/HTTP/Public/HttpManager.h
--- a/Engine/Source/Runtime/Online/HTTP/Public/HttpManager.h
+++ b/Engine/Source/Runtime/Online/HTTP/Public/HttpManager.h
@@ -104,6 +104,29 @@ public:
 	 */
 	virtual bool SupportsDynamicProxy() const;
 
+	/**
+	 * Set the method used to set a Correlation id on each request, if one is not already specified.
+	 *
+	 * This method allows you to override the Engine default method.
+	 *
+	 * @param InCorrelationIdMethod The method to use when sending a request, if no Correlation id is already set
+	 */
+	void SetCorrelationIdMethod(TFunction<FString()> InCorrelationIdMethod);
+
+	/**
+	 * Create a new correlation id for a request
+	 *
+	 * @return The new correlationid string
+	 */
+	FString CreateCorrelationId() const;
+
+	/**
+	 * Get the default method for creating new correlation ids for a request
+	 *
+	 * @return The default correlationid creation method
+	 */
+	static TFunction<FString()> GetDefaultCorrelationIdMethod();
+
 protected:
 	/** 
 	 * Create HTTP thread object
@@ -140,6 +163,9 @@ protected:
 
 	FHttpThread* Thread;
 
+	/** This method will be called to generate a CorrelationId on all requests being sent if one is not already set */
+	TFunction<FString()> CorrelationIdMethod;
+
 PACKAGE_SCOPE:
 
 	/** Used to lock access to add/remove/find requests */
diff --git a/Engine/Source/Runtime/Online/HTTP/Public/HttpModule.h b/Engine/Source/Runtime/Online/HTTP/Public/HttpModule.h
--- a/Engine/Source/Runtime/Online/HTTP/Public/HttpModule.h
+++ b/Engine/Source/Runtime/Online/HTTP/Public/HttpModule.h
@@ -239,6 +239,20 @@ public:
 		HttpThreadIdleMinimumSleepTimeInSeconds = InHttpThreadIdleMinimumSleepTimeInSeconds;
 	}
 
+	/**
+	 * Get the default headers that are appended to every request
+	 * @return the default headers
+	 */
+	const TMap<FString, FString>& GetDefaultHeaders() const { return DefaultHeaders; }
+
+	/**
+	 * Add a default header to be appended to future requests
+	 * If a request already specifies this header, then the defaulted version will not be used
+	 * @param HeaderName - Name of the header (e.g., "Content-Type")
+	 * @param HeaderValue - Value of the header
+	 */
+	void AddDefaultHeader(const FString& HeaderName, const FString& HeaderValue) { DefaultHeaders.Emplace(HeaderName, HeaderValue); }
+
 	/**
 	 * @returns The current proxy address.
 	 */
@@ -321,6 +335,8 @@ private:
 	bool bEnableHttp;
 	/** toggles null (mock) http requests */
 	bool bUseNullHttp;
+	/** Default headers - each request will include these headers, using the default value if not overridden */
+	TMap<FString, FString> DefaultHeaders;
 	/** singleton for the module while loaded and available */
 	static FHttpModule* Singleton;
 	/** The address to use for proxy, in format IPADDRESS:PORT */
diff --git a/Engine/Source/Runtime/Online/HTTP/Public/HttpRetrySystem.h b/Engine/Source/Runtime/Online/HTTP/Public/HttpRetrySystem.h
--- a/Engine/Source/Runtime/Online/HTTP/Public/HttpRetrySystem.h
+++ b/Engine/Source/Runtime/Online/HTTP/Public/HttpRetrySystem.h
@@ -3,6 +3,7 @@
 #pragma once
 
 #include "CoreMinimal.h"
+#include "Templates/Atomic.h"
 #include "Interfaces/IHttpRequest.h"
 #include "HttpRequestAdapter.h"
 
@@ -177,6 +178,7 @@ namespace FHttpRetrySystem
          */
         HTTP_API bool Update(uint32* FileCount = NULL, uint32* FailingCount = NULL, uint32* FailedCount = NULL, uint32* CompletedCount = NULL);
 		HTTP_API void SetRandomFailureRate(float Value) { RandomFailureRate = FRandomFailureRateSetting::Create(Value); }
+		HTTP_API void SetDefaultRetryLimit(uint32 Value) { RetryLimitCountDefault = FRetryLimitCountSetting::Create(Value); }
 
     protected:
 		friend class FRequest;
diff --git a/Engine/Source/Runtime/Online/ImageDownload/Private/WebImageCache.cpp b/Engine/Source/Runtime/Online/ImageDownload/Private/WebImageCache.cpp
--- a/Engine/Source/Runtime/Online/ImageDownload/Private/WebImageCache.cpp
+++ b/Engine/Source/Runtime/Online/ImageDownload/Private/WebImageCache.cpp
@@ -16,6 +16,17 @@ void FWebImageCache::PreUnload()
 	}
 }
 
+void FWebImageCache::Empty()
+{
+	for (const TSharedRef<FWebImage>& WebImage : StrongRefCache)
+	{
+		WebImage->CancelDownload();
+	}
+
+	UrlToImageMap.Empty();
+	StrongRefCache.Empty();
+}
+
 TSharedRef<const FWebImage> FWebImageCache::Download(const FString& Url, const TOptional<FString>& DefaultImageUrl)
 {
 	TAttribute<const FSlateBrush*> StandInBrush;
diff --git a/Engine/Source/Runtime/Online/ImageDownload/Public/WebImageCache.h b/Engine/Source/Runtime/Online/ImageDownload/Public/WebImageCache.h
--- a/Engine/Source/Runtime/Online/ImageDownload/Public/WebImageCache.h
+++ b/Engine/Source/Runtime/Online/ImageDownload/Public/WebImageCache.h
@@ -22,6 +22,9 @@ public:
 	/** Signifies the module is being unloaded and to perform any actions that depend on other modules which may be unloaded as well */
 	void PreUnload();
 
+	/** Removes all cached images */
+	void Empty();
+
 	/** Find or create a WebImage object for this URL (you probably just want to call ->Attr() on this) */
 	TSharedRef<const FWebImage> Download(const FString& Url, const TOptional<FString>& DefaultImageUrl = TOptional<FString>());
 
diff --git a/Engine/Source/Runtime/Slate/Private/Framework/Application/SlateApplication.cpp b/Engine/Source/Runtime/Slate/Private/Framework/Application/SlateApplication.cpp
--- a/Engine/Source/Runtime/Slate/Private/Framework/Application/SlateApplication.cpp
+++ b/Engine/Source/Runtime/Slate/Private/Framework/Application/SlateApplication.cpp
@@ -39,7 +39,9 @@
 #include "HAL/LowLevelMemTracker.h"
 #include "ProfilingDebugging/CsvProfiler.h"
 
-#define SLATE_HAS_WIDGET_REFLECTOR !(UE_BUILD_TEST || UE_BUILD_SHIPPING) && PLATFORM_DESKTOP
+#ifndef SLATE_HAS_WIDGET_REFLECTOR
+	#define SLATE_HAS_WIDGET_REFLECTOR !(UE_BUILD_TEST || UE_BUILD_SHIPPING) && PLATFORM_DESKTOP
+#endif
 
 #if PLATFORM_WINDOWS
 #include "Windows/WindowsHWrapper.h"
@@ -7341,4 +7343,3 @@ void FSlateApplication::InputPreProcessorsHelper::RemoveAll()
 	InputPreProcessorList.Reset();
 }
 
-#undef SLATE_HAS_WIDGET_REFLECTOR
diff --git a/Engine/Source/Runtime/Slate/Private/Framework/Text/ShapedTextCache.cpp b/Engine/Source/Runtime/Slate/Private/Framework/Text/ShapedTextCache.cpp
--- a/Engine/Source/Runtime/Slate/Private/Framework/Text/ShapedTextCache.cpp
+++ b/Engine/Source/Runtime/Slate/Private/Framework/Text/ShapedTextCache.cpp
@@ -19,7 +19,14 @@ FShapedGlyphSequencePtr FShapedTextCache::FindShapedText(const FCachedShapedText
 
 FShapedGlyphSequenceRef FShapedTextCache::AddShapedText(const FCachedShapedTextKey& InKey, const TCHAR* InText)
 {
-	FShapedGlyphSequenceRef ShapedText = FontCache.ShapeBidirectionalText(
+	const TSharedPtr<FSlateFontCache> FontCache = FontCachePtr.Pin();
+
+	if (!FontCache.IsValid())
+	{
+		return MakeShareable(new FShapedGlyphSequence());
+	}
+
+	FShapedGlyphSequenceRef ShapedText = FontCache->ShapeBidirectionalText(
 		InText, 
 		InKey.TextRange.BeginIndex, 
 		InKey.TextRange.Len(), 
@@ -34,7 +41,14 @@ FShapedGlyphSequenceRef FShapedTextCache::AddShapedText(const FCachedShapedTextK
 
 FShapedGlyphSequenceRef FShapedTextCache::AddShapedText(const FCachedShapedTextKey& InKey, const TCHAR* InText, const TextBiDi::ETextDirection InTextDirection)
 {
-	FShapedGlyphSequenceRef ShapedText = FontCache.ShapeUnidirectionalText(
+	const TSharedPtr<FSlateFontCache> FontCache = FontCachePtr.Pin();
+
+	if (!FontCache.IsValid())
+	{
+		return MakeShareable(new FShapedGlyphSequence());
+	}
+
+	FShapedGlyphSequenceRef ShapedText = FontCache->ShapeUnidirectionalText(
 		InText, 
 		InKey.TextRange.BeginIndex, 
 		InKey.TextRange.Len(), 
@@ -105,12 +119,17 @@ FVector2D ShapedTextCacheUtil::MeasureShapedText(const FShapedTextCacheRef& InSh
 
 int32 ShapedTextCacheUtil::FindCharacterIndexAtOffset(const FShapedTextCacheRef& InShapedTextCache, const FCachedShapedTextKey& InRunKey, const FTextRange& InTextRange, const TCHAR* InText, const int32 InHorizontalOffset)
 {
-	FSlateFontCache& FontCache = InShapedTextCache->GetFontCache();
+	TSharedPtr<FSlateFontCache> FontCache = InShapedTextCache->GetFontCache();
+
+	if (!FontCache.IsValid())
+	{
+		return INDEX_NONE;
+	}
 
 	// Get the shaped text for the entire run and try and take a sub-measurement from it - this can help minimize the amount of text shaping that needs to be done when measuring text
 	FShapedGlyphSequenceRef ShapedText = InShapedTextCache->FindOrAddShapedText(InRunKey, InText);
 
-	TOptional<FShapedGlyphSequence::FGlyphOffsetResult> GlyphOffsetResult = ShapedText->GetGlyphAtOffset(FontCache, InTextRange.BeginIndex, InTextRange.EndIndex, InHorizontalOffset);
+	TOptional<FShapedGlyphSequence::FGlyphOffsetResult> GlyphOffsetResult = ShapedText->GetGlyphAtOffset(*FontCache, InTextRange.BeginIndex, InTextRange.EndIndex, InHorizontalOffset);
 	if (!GlyphOffsetResult.IsSet())
 	{
 		FCachedShapedTextKey IndexAtOffsetKey = InRunKey;
@@ -118,7 +137,7 @@ int32 ShapedTextCacheUtil::FindCharacterIndexAtOffset(const FShapedTextCacheRef&
 
 		// Couldn't search the sub-range, try and search from a shape of the specified range
 		ShapedText = InShapedTextCache->FindOrAddShapedText(IndexAtOffsetKey, InText);
-		GlyphOffsetResult = ShapedText->GetGlyphAtOffset(FontCache, InHorizontalOffset);
+		GlyphOffsetResult = ShapedText->GetGlyphAtOffset(*FontCache, InHorizontalOffset);
 	}
 
 	check(GlyphOffsetResult.IsSet());
@@ -143,7 +162,7 @@ int32 ShapedTextCacheUtil::FindCharacterIndexAtOffset(const FShapedTextCacheRef&
 			{
 				FShapedGlyphSequenceRef GraphemeShapedText = GetShapedTextSubSequence(InShapedTextCache, LigatureKey, FTextRange(PrevCharIndex, CurrentCharIndex), *LigatureString, GlyphOffsetResultValue.Glyph->TextDirection);
 				
-				const FShapedGlyphSequence::FGlyphOffsetResult GraphemeOffsetResult = GraphemeShapedText->GetGlyphAtOffset(FontCache, InHorizontalOffset, CurrentOffset);
+				const FShapedGlyphSequence::FGlyphOffsetResult GraphemeOffsetResult = GraphemeShapedText->GetGlyphAtOffset(*FontCache, InHorizontalOffset, CurrentOffset);
 				if (GraphemeOffsetResult.Glyph)
 				{
 					return GlyphOffsetResultValue.Glyph->SourceIndex + GraphemeOffsetResult.CharacterIndex;
@@ -162,7 +181,7 @@ int32 ShapedTextCacheUtil::FindCharacterIndexAtOffset(const FShapedTextCacheRef&
 			{
 				FShapedGlyphSequenceRef GraphemeShapedText = GetShapedTextSubSequence(InShapedTextCache, LigatureKey, FTextRange(CurrentCharIndex, PrevCharIndex), *LigatureString, GlyphOffsetResultValue.Glyph->TextDirection);
 
-				const FShapedGlyphSequence::FGlyphOffsetResult GraphemeOffsetResult = GraphemeShapedText->GetGlyphAtOffset(FontCache, InHorizontalOffset, CurrentOffset);
+				const FShapedGlyphSequence::FGlyphOffsetResult GraphemeOffsetResult = GraphemeShapedText->GetGlyphAtOffset(*FontCache, InHorizontalOffset, CurrentOffset);
 				if (GraphemeOffsetResult.Glyph)
 				{
 					return GlyphOffsetResultValue.Glyph->SourceIndex + ((PrevCharIndex != INDEX_NONE) ? PrevCharIndex : GraphemeOffsetResult.CharacterIndex);
diff --git a/Engine/Source/Runtime/Slate/Private/Framework/Text/TextLayout.cpp b/Engine/Source/Runtime/Slate/Private/Framework/Text/TextLayout.cpp
--- a/Engine/Source/Runtime/Slate/Private/Framework/Text/TextLayout.cpp
+++ b/Engine/Source/Runtime/Slate/Private/Framework/Text/TextLayout.cpp
@@ -2587,7 +2587,7 @@ FTextLayout::~FTextLayout()
 
 FTextLayout::FLineModel::FLineModel( const TSharedRef< FString >& InText ) 
 	: Text( InText )
-	, ShapedTextCache( FShapedTextCache::Create(*FSlateApplication::Get().GetRenderer()->GetFontCache()) )
+	, ShapedTextCache( FShapedTextCache::Create(FSlateApplication::Get().GetRenderer()->GetFontCache()))
 	, TextBaseDirection( TextBiDi::ETextDirection::LeftToRight )
 	, Runs()
 	, BreakCandidates()
diff --git a/Engine/Source/Runtime/Slate/Private/Widgets/Layout/SScrollBarTrack.cpp b/Engine/Source/Runtime/Slate/Private/Widgets/Layout/SScrollBarTrack.cpp
--- a/Engine/Source/Runtime/Slate/Private/Widgets/Layout/SScrollBarTrack.cpp
+++ b/Engine/Source/Runtime/Slate/Private/Widgets/Layout/SScrollBarTrack.cpp
@@ -8,7 +8,7 @@ static const int32 NUM_SCROLLBAR_SLOTS = 3;
 void SScrollBarTrack::Construct(const FArguments& InArgs)
 {
 	OffsetFraction = 0;
-	ThumbSizeFraction = 0;
+	ThumbSizeFraction = 1.0; // default to zero offset from the top with a full track thumb (scrolling not needed)
 	MinThumbSize = 35;
 	Orientation = InArgs._Orientation;
 
@@ -117,6 +117,12 @@ void SScrollBarTrack::SetSizes(float InThumbOffsetFraction, float InThumbSizeFra
 {
 	OffsetFraction = InThumbOffsetFraction;
 	ThumbSizeFraction = InThumbSizeFraction;
+
+	// If you have no thumb, then it's effectively the size of the whole track
+	if (ThumbSizeFraction == 0.0f)
+	{
+		ThumbSizeFraction = 1.0f;
+	}
 }
 
 bool SScrollBarTrack::IsNeeded() const
diff --git a/Engine/Source/Runtime/Slate/Public/Framework/Text/ShapedTextCache.h b/Engine/Source/Runtime/Slate/Public/Framework/Text/ShapedTextCache.h
--- a/Engine/Source/Runtime/Slate/Public/Framework/Text/ShapedTextCache.h
+++ b/Engine/Source/Runtime/Slate/Public/Framework/Text/ShapedTextCache.h
@@ -56,7 +56,7 @@ class SLATE_API FShapedTextCache
 {
 public:
 	/** Create a new shaped text cache */
-	static FShapedTextCacheRef Create(FSlateFontCache& InFontCache)
+	static FShapedTextCacheRef Create(const TSharedRef<FSlateFontCache>& InFontCache)
 	{
 		return MakeShareable(new FShapedTextCache(InFontCache));
 	}
@@ -104,20 +104,20 @@ public:
 	/**
 	 * Get the font cache used by this instance
 	 */
-	FSlateFontCache& GetFontCache() const
+	TSharedPtr<FSlateFontCache> GetFontCache() const
 	{
-		return FontCache;
+		return FontCachePtr.Pin();
 	}
 
 private:
 	/** Constructor */
-	FShapedTextCache(FSlateFontCache& InFontCache)
-		: FontCache(InFontCache)
+	FShapedTextCache(const TSharedRef<FSlateFontCache>& InFontCache)
+		: FontCachePtr(InFontCache)
 	{
 	}
 
 	/** Font cache to use when shaping text */
-	FSlateFontCache& FontCache;
+	TWeakPtr<FSlateFontCache> FontCachePtr;
 
 	/** Mapping between a cache key and the corresponding shaped text */
 	TMap<FCachedShapedTextKey, FShapedGlyphSequencePtr> CachedShapedText;
diff --git a/Engine/Source/Runtime/SlateCore/Private/Rendering/ElementBatcher.cpp b/Engine/Source/Runtime/SlateCore/Private/Rendering/ElementBatcher.cpp
--- a/Engine/Source/Runtime/SlateCore/Private/Rendering/ElementBatcher.cpp
+++ b/Engine/Source/Runtime/SlateCore/Private/Rendering/ElementBatcher.cpp
@@ -1888,7 +1888,8 @@ void FSlateElementBatcher::AddViewportElement( const FSlateDrawElement& DrawElem
 	// If the viewport disallows scaling, force size to current texture size.
 	if (ViewportResource != nullptr && !InPayload.bAllowViewportScaling)
 	{
-		BotRight = FVector2D(ViewportResource->GetWidth(), ViewportResource->GetHeight());
+		const float ElementScale = DrawElement.GetScale();
+		BotRight = FVector2D(ViewportResource->GetWidth() / ElementScale, ViewportResource->GetHeight() / ElementScale);
 	}
 
 	FVector2D TopRight = FVector2D( BotRight.X, TopLeft.Y);
diff --git a/Engine/Source/Runtime/SlateCore/Public/Fonts/FontCache.h b/Engine/Source/Runtime/SlateCore/Public/Fonts/FontCache.h
--- a/Engine/Source/Runtime/SlateCore/Public/Fonts/FontCache.h
+++ b/Engine/Source/Runtime/SlateCore/Public/Fonts/FontCache.h
@@ -228,6 +228,17 @@ public:
 		int32 TextLen;
 	};
 
+	explicit FShapedGlyphSequence()
+		: GlyphsToRender()
+		, TextBaseline(0)
+		, MaxTextHeight(0)
+		, FontMaterial(nullptr)
+		, OutlineSettings()
+		, SequenceWidth(0)
+		, GlyphFontFaces()
+		, SourceIndicesToGlyphData(FSourceTextRange(0, 0))
+	{ }
+
 	FShapedGlyphSequence(TArray<FShapedGlyphEntry> InGlyphsToRender, const int16 InTextBaseline, const uint16 InMaxTextHeight, const UObject* InFontMaterial, const FFontOutlineSettings& InOutlineSettings, const FSourceTextRange& InSourceTextRange);
 	~FShapedGlyphSequence();
 
diff --git a/Engine/Source/Runtime/WebBrowser/Private/Android/AndroidWebBrowserWindow.cpp b/Engine/Source/Runtime/WebBrowser/Private/Android/AndroidWebBrowserWindow.cpp
--- a/Engine/Source/Runtime/WebBrowser/Private/Android/AndroidWebBrowserWindow.cpp
+++ b/Engine/Source/Runtime/WebBrowser/Private/Android/AndroidWebBrowserWindow.cpp
@@ -144,6 +144,15 @@ void FAndroidWebBrowserWindow::OnMouseLeave(const FPointerEvent& MouseEvent)
 {
 }
 
+void FAndroidWebBrowserWindow::SetSupportsMouseWheel(bool bValue)
+{
+}
+
+bool FAndroidWebBrowserWindow::GetSupportsMouseWheel() const
+{
+	return false;
+}
+
 FReply FAndroidWebBrowserWindow::OnMouseWheel(const FGeometry& MyGeometry, const FPointerEvent& MouseEvent, bool bIsPopup)
 {
 	return FReply::Unhandled();
diff --git a/Engine/Source/Runtime/WebBrowser/Private/Android/AndroidWebBrowserWindow.h b/Engine/Source/Runtime/WebBrowser/Private/Android/AndroidWebBrowserWindow.h
--- a/Engine/Source/Runtime/WebBrowser/Private/Android/AndroidWebBrowserWindow.h
+++ b/Engine/Source/Runtime/WebBrowser/Private/Android/AndroidWebBrowserWindow.h
@@ -69,6 +69,8 @@ public:
 	virtual FReply OnMouseButtonDoubleClick(const FGeometry& MyGeometry, const FPointerEvent& MouseEvent, bool bIsPopup) override;
 	virtual FReply OnMouseMove(const FGeometry& MyGeometry, const FPointerEvent& MouseEvent, bool bIsPopup) override;
 	virtual void OnMouseLeave(const FPointerEvent& MouseEvent) override;
+	virtual void SetSupportsMouseWheel(bool bValue) override;
+	virtual bool GetSupportsMouseWheel() const override;
 	virtual FReply OnMouseWheel(const FGeometry& MyGeometry, const FPointerEvent& MouseEvent, bool bIsPopup) override;
 	virtual void OnFocus(bool SetFocus, bool bIsPopup) override;
 	virtual void OnCaptureLost() override;
@@ -176,6 +178,11 @@ public:
 		return SuppressContextMenuDelgate;
 	}
 
+	virtual FOnDragWindow& OnDragWindow() override
+	{
+		return DragWindowDelegate;
+	}
+
 public:
 
 	/**
@@ -255,6 +262,9 @@ private:
 	/** Delegate for suppressing context menu */
 	FOnSuppressContextMenu SuppressContextMenuDelgate;
 
+	/** Delegate that is executed when a drag event is detected in an area of the web page tagged as a drag region. */
+	FOnDragWindow DragWindowDelegate;
+
 	/** Current state of the document being loaded. */
 	EWebBrowserDocumentState DocumentState;
 	int ErrorCode;
diff --git a/Engine/Source/Runtime/WebBrowser/Private/CEF/CEFBrowserHandler.cpp b/Engine/Source/Runtime/WebBrowser/Private/CEF/CEFBrowserHandler.cpp
--- a/Engine/Source/Runtime/WebBrowser/Private/CEF/CEFBrowserHandler.cpp
+++ b/Engine/Source/Runtime/WebBrowser/Private/CEF/CEFBrowserHandler.cpp
@@ -14,7 +14,7 @@
 #include "CEFBrowserByteResource.h"
 #include "Framework/Application/SlateApplication.h"
 #include "HAL/ThreadingBase.h"
-
+#include "PlatformHttp.h"
 
 #define LOCTEXT_NAMESPACE "WebBrowserHandler"
 
@@ -22,8 +22,9 @@
 // Used to force returning custom content instead of performing a request.
 const FString CustomContentMethod(TEXT("X-GET-CUSTOM-CONTENT"));
 
-FCEFBrowserHandler::FCEFBrowserHandler(bool InUseTransparency)
-: bUseTransparency(InUseTransparency)
+FCEFBrowserHandler::FCEFBrowserHandler(bool InUseTransparency, const TArray<FString>& InAltRetryDomains)
+: bUseTransparency(InUseTransparency),
+AltRetryDomains(InAltRetryDomains)
 { }
 
 void FCEFBrowserHandler::OnTitleChange(CefRefPtr<CefBrowser> Browser, const CefString& Title)
@@ -204,6 +205,7 @@ void FCEFBrowserHandler::OnLoadError(CefRefPtr<CefBrowser> Browser,
 	const CefString& ErrorText,
 	const CefString& FailedUrl)
 {
+
 	// notify browser window
 	if (Frame->IsMain())
 	{
@@ -211,6 +213,18 @@ void FCEFBrowserHandler::OnLoadError(CefRefPtr<CefBrowser> Browser,
 
 		if (BrowserWindow.IsValid())
 		{
+			if (AltRetryDomains.Num() > 0 && AltRetryDomainIdx < (uint32)AltRetryDomains.Num())
+			{
+				FString Url = FailedUrl.ToWString().c_str();
+				FString OriginalUrlDomain = FPlatformHttp::GetUrlDomain(Url);
+				if (!OriginalUrlDomain.IsEmpty())
+				{
+					const FString NewUrl(Url.Replace(*OriginalUrlDomain, *AltRetryDomains[AltRetryDomainIdx++]));
+					BrowserWindow->LoadURL(NewUrl);
+					return;
+				}
+
+			}
 			BrowserWindow->NotifyDocumentError(InErrorCode, ErrorText, FailedUrl);
 		}
 	}
@@ -610,7 +624,6 @@ void FCEFBrowserHandler::OnResetDialogState(CefRefPtr<CefBrowser> Browser)
 	}
 }
 
-
 void FCEFBrowserHandler::OnBeforeContextMenu(CefRefPtr<CefBrowser> Browser, CefRefPtr<CefFrame> Frame, CefRefPtr<CefContextMenuParams> Params, CefRefPtr<CefMenuModel> Model)
 {
 	TSharedPtr<FCEFWebBrowserWindow> BrowserWindow = BrowserWindowPtr.Pin();
@@ -620,6 +633,22 @@ void FCEFBrowserHandler::OnBeforeContextMenu(CefRefPtr<CefBrowser> Browser, CefR
 	}
 }
 
+void FCEFBrowserHandler::OnDraggableRegionsChanged(CefRefPtr<CefBrowser> Browser, const std::vector<CefDraggableRegion>& Regions)
+{
+	TSharedPtr<FCEFWebBrowserWindow> BrowserWindow = BrowserWindowPtr.Pin();
+	if (BrowserWindow.IsValid())
+	{
+		TArray<FWebBrowserDragRegion> DragRegions;
+		for (uint32 Idx = 0; Idx < Regions.size(); Idx++)
+		{
+			DragRegions.Add(FWebBrowserDragRegion(
+				FIntRect(Regions[Idx].bounds.x, Regions[Idx].bounds.y, Regions[Idx].bounds.x + Regions[Idx].bounds.width, Regions[Idx].bounds.y + Regions[Idx].bounds.height),
+				Regions[Idx].draggable ? true : false));
+		}
+		BrowserWindow->UpdateDragRegions(DragRegions);
+	}
+}
+
 #undef LOCTEXT_NAMESPACE
 
 #endif // WITH_CEF
diff --git a/Engine/Source/Runtime/WebBrowser/Private/CEF/CEFBrowserHandler.h b/Engine/Source/Runtime/WebBrowser/Private/CEF/CEFBrowserHandler.h
--- a/Engine/Source/Runtime/WebBrowser/Private/CEF/CEFBrowserHandler.h
+++ b/Engine/Source/Runtime/WebBrowser/Private/CEF/CEFBrowserHandler.h
@@ -54,11 +54,12 @@ class FCEFBrowserHandler
 	, public CefKeyboardHandler
 	, public CefJSDialogHandler
 	, public CefContextMenuHandler
+	, public CefDragHandler
 {
 public:
 
 	/** Default constructor. */
-	FCEFBrowserHandler(bool InUseTransparency);
+	FCEFBrowserHandler(bool InUseTransparency, const TArray<FString>& AltRetryDomains = TArray<FString>());
 
 public:
 
@@ -123,6 +124,11 @@ public:
 		return this;
 	}
 
+	virtual CefRefPtr<CefDragHandler> GetDragHandler() override
+	{
+		return this;
+	}
+
 
 	virtual bool OnProcessMessageReceived(CefRefPtr<CefBrowser> Browser,
 		CefProcessId SourceProcess,
@@ -280,6 +286,21 @@ public:
 
 	virtual void OnResetDialogState(CefRefPtr<CefBrowser> Browser) override;
 
+public:
+	// CefContextMenuHandler
+
+	virtual void OnBeforeContextMenu(CefRefPtr<CefBrowser> Browser,
+		CefRefPtr<CefFrame> Frame,
+		CefRefPtr<CefContextMenuParams> Params,
+		CefRefPtr<CefMenuModel> Model) override;
+
+public:
+	// CefDragHandler interface
+
+	virtual void OnDraggableRegionsChanged(
+		CefRefPtr<CefBrowser> Browser,
+		const std::vector<CefDraggableRegion>& Regions) override;
+
 public:
 
 	IWebBrowserWindow::FOnBeforePopupDelegate& OnBeforePopup()
@@ -292,18 +313,15 @@ public:
 		return CreateWindowDelegate;
 	}
 
-public:
-	virtual void OnBeforeContextMenu(CefRefPtr<CefBrowser> Browser, 
- 		CefRefPtr<CefFrame> Frame,
-		CefRefPtr<CefContextMenuParams> Params,
-		CefRefPtr<CefMenuModel> Model) override;
-
 private:
 
 	bool ShowDevTools(const CefRefPtr<CefBrowser>& Browser);
 
 	bool bUseTransparency;
 
+	TArray<FString> AltRetryDomains;
+	uint32 AltRetryDomainIdx = 0;
+
 	/** Delegate for notifying that a popup window is attempting to open. */
 	IWebBrowserWindow::FOnBeforePopupDelegate BeforePopupDelegate;
 	
diff --git a/Engine/Source/Runtime/WebBrowser/Private/CEF/CEFImeHandler.h b/Engine/Source/Runtime/WebBrowser/Private/CEF/CEFImeHandler.h
--- a/Engine/Source/Runtime/WebBrowser/Private/CEF/CEFImeHandler.h
+++ b/Engine/Source/Runtime/WebBrowser/Private/CEF/CEFImeHandler.h
@@ -29,10 +29,12 @@ THIRD_PARTY_INCLUDES_END
 #include "Windows/HideWindowsPlatformAtomics.h"
 #include "Windows/HideWindowsPlatformTypes.h"
 #endif
+#include "Layout/Geometry.h"
 
 class ITextInputMethodSystem;
 class FCEFTextInputMethodContext;
 class ITextInputMethodChangeNotifier;
+class SWidget;
 
 class FCEFImeHandler
 	: public TSharedFromThis<FCEFImeHandler>
diff --git a/Engine/Source/Runtime/WebBrowser/Private/CEF/CEFTextInputMethodContext.h b/Engine/Source/Runtime/WebBrowser/Private/CEF/CEFTextInputMethodContext.h
--- a/Engine/Source/Runtime/WebBrowser/Private/CEF/CEFTextInputMethodContext.h
+++ b/Engine/Source/Runtime/WebBrowser/Private/CEF/CEFTextInputMethodContext.h
@@ -34,9 +34,11 @@ THIRD_PARTY_INCLUDES_END
 #endif
 
 #include "GenericPlatform/ITextInputMethodSystem.h"
+#include "Layout/Geometry.h"
 
 class FCEFWebBrowserWindow;
 class FCEFImeHandler;
+class SWindow;
 
 class FCEFTextInputMethodContext : public ITextInputMethodContext
 {
diff --git a/Engine/Source/Runtime/WebBrowser/Private/CEF/CEFWebBrowserWindow.cpp b/Engine/Source/Runtime/WebBrowser/Private/CEF/CEFWebBrowserWindow.cpp
--- a/Engine/Source/Runtime/WebBrowser/Private/CEF/CEFWebBrowserWindow.cpp
+++ b/Engine/Source/Runtime/WebBrowser/Private/CEF/CEFWebBrowserWindow.cpp
@@ -5,6 +5,7 @@
 #include "UObject/Stack.h"
 #include "Framework/Application/SlateApplication.h"
 #include "Textures/SlateUpdatableTexture.h"
+#include "HAL/PlatformApplicationMisc.h"
 
 #if WITH_CEF3
 
@@ -13,6 +14,7 @@
 #include "CEFBrowserClosureTask.h"
 #include "CEFJSScripting.h"
 #include "CEFImeHandler.h"
+#include "Async/Async.h"
 
 #if PLATFORM_MAC
 // Needed for character code definitions
@@ -371,6 +373,7 @@ FCEFWebBrowserWindow::FCEFWebBrowserWindow(CefRefPtr<CefBrowser> InBrowser, CefR
 	, WebBrowserHandler(InHandler)
 	, CurrentUrl(InUrl)
 	, ViewportSize(FIntPoint::ZeroValue)
+	, ViewportDPIScaleFactor(1.0f)
 	, bIsClosing(false)
 	, bIsInitialized(false)
 	, ContentsToLoad(InContentsToLoad)
@@ -382,6 +385,7 @@ FCEFWebBrowserWindow::FCEFWebBrowserWindow(CefRefPtr<CefBrowser> InBrowser, CefR
 	, bIsHidden(false)
 	, bTickedLastFrame(true)
 	, bNeedsResize(false)
+	, bDraggingWindow(false)
 	, PreviousKeyDownEvent()
 	, PreviousKeyUpEvent()
 	, PreviousCharacterEvent()
@@ -390,6 +394,7 @@ FCEFWebBrowserWindow::FCEFWebBrowserWindow(CefRefPtr<CefBrowser> InBrowser, CefR
 	, bIgnoreCharacterEvent(false)
 	, bMainHasFocus(false)
 	, bPopupHasFocus(false)
+	, bSupportsMouseWheel(true)
 	, bRecoverFromRenderProcessCrash(false)
 	, ErrorCode(0)
 	, bDeferNavigations(false)
@@ -400,29 +405,90 @@ FCEFWebBrowserWindow::FCEFWebBrowserWindow(CefRefPtr<CefBrowser> InBrowser, CefR
 {
 	check(InBrowser.get() != nullptr);
 
+	UpdatableTextures[0] = nullptr;
+	UpdatableTextures[1] = nullptr;
+
+	if (!CreateInitialTextures())
+	{
+		ReleaseTextures();
+	}
+
+#if USE_BUFFERED_VIDEO
+	BufferedVideo = TUniquePtr<FBrowserBufferedVideo>(new FBrowserBufferedVideo(4));
+#endif
+}
+
+void FCEFWebBrowserWindow::ReleaseTextures()
+{
+	for (int I = 0; I < 2; ++I)
+	{
+		if (UpdatableTextures[I] != nullptr)
+		{
+			FSlateUpdatableTexture* TextureToRelease = UpdatableTextures[I];
+			AsyncTask(ENamedThreads::GameThread, [TextureToRelease]()
+			{
+				if (FSlateApplication::IsInitialized())
+				{
+					if (FSlateRenderer* Renderer = FSlateApplication::Get().GetRenderer())
+					{
+						Renderer->ReleaseUpdatableTexture(TextureToRelease);
+					}
+				}
+			});
+
+			UpdatableTextures[I] = nullptr;
+		}
+	}
+}
+
+bool FCEFWebBrowserWindow::CreateInitialTextures()
+{
 	if (FSlateApplication::IsInitialized())
 	{
 		if (FSlateRenderer* Renderer = FSlateApplication::Get().GetRenderer())
 		{
+			if (Renderer->HasLostDevice())
+			{
+				return false;
+			}
+
 			// Create a transparent dummy texture for our buffers which will prevent slate from applying an 
 			// undesirable quad if it happens to ask for this buffer before we get a chance to paint to it.
 			TArray<uint8> RawData;
 			RawData.AddZeroed(4);
 			UpdatableTextures[0] = Renderer->CreateUpdatableTexture(1, 1);
+
+			if (Renderer->HasLostDevice())
+			{
+				return false;
+			}
+
 			UpdatableTextures[0]->UpdateTextureThreadSafeRaw(1, 1, RawData.GetData());
+
+			if (Renderer->HasLostDevice())
+			{
+				return false;
+			}
+
 			UpdatableTextures[1] = Renderer->CreateUpdatableTexture(1, 1);
+
+			if (Renderer->HasLostDevice())
+			{
+				return false;
+			}
+
 			UpdatableTextures[1]->UpdateTextureThreadSafeRaw(1, 1, RawData.GetData());
+
+			if (Renderer->HasLostDevice())
+			{
+				return false;
+			}
+
+			return true;
 		}
 	}
-	else
-	{
-		UpdatableTextures[0] = nullptr;
-		UpdatableTextures[1] = nullptr;
-	}
 
-#if USE_BUFFERED_VIDEO
-	BufferedVideo = TUniquePtr<FBrowserBufferedVideo>(new FBrowserBufferedVideo(4));
-#endif
+	return false;
 }
 
 FCEFWebBrowserWindow::~FCEFWebBrowserWindow()
@@ -431,21 +497,7 @@ FCEFWebBrowserWindow::~FCEFWebBrowserWindow()
 	WebBrowserHandler->OnBeforePopup().Unbind();
 	CloseBrowser(true);
 
-	if (FSlateApplication::IsInitialized())
-	{
-		if (FSlateRenderer* Renderer = FSlateApplication::Get().GetRenderer())
-		{
-			for (int I = 0; I < 1; ++I)
-			{
-				if (UpdatableTextures[I] != nullptr)
-				{
-					Renderer->ReleaseUpdatableTexture(UpdatableTextures[I]);
-				}
-			}
-		}
-	}
-	UpdatableTextures[0] = nullptr;
-	UpdatableTextures[1] = nullptr;
+	ReleaseTextures();
 
 	BufferedVideo.Reset();
 }
@@ -484,11 +536,14 @@ void FCEFWebBrowserWindow::SetViewportSize(FIntPoint WindowSize, FIntPoint Windo
 	}
 	bTickedLastFrame=true;
 
+	const float WindowDPIScaleFactor = ParentWindow.IsValid() ? ParentWindow->GetNativeWindow()->GetDPIScaleFactor() : 1.0f;
+
 	// Ignore sizes that can't be seen as it forces CEF to re-render whole image
-	if (WindowSize.X > 0 && WindowSize.Y > 0 && ViewportSize != WindowSize)
+	if ((WindowSize.X > 0 && WindowSize.Y > 0 && ViewportSize != WindowSize) || WindowDPIScaleFactor != ViewportDPIScaleFactor)
 	{
 		bool bFirstSize = ViewportSize == FIntPoint::ZeroValue;
 		ViewportSize = WindowSize;
+		ViewportDPIScaleFactor = WindowDPIScaleFactor;
 
 		if (IsValid())
 		{
@@ -1152,12 +1207,19 @@ FReply FCEFWebBrowserWindow::OnMouseButtonDown(const FGeometry& MyGeometry, cons
 
 		if(bIsCefSupportedButton)
 		{
-		CefBrowserHost::MouseButtonType Type =
-			(Button == EKeys::LeftMouseButton ? MBT_LEFT : (
-			Button == EKeys::RightMouseButton ? MBT_RIGHT : MBT_MIDDLE));
+			CefBrowserHost::MouseButtonType Type =
+				(Button == EKeys::LeftMouseButton ? MBT_LEFT : (
+				Button == EKeys::RightMouseButton ? MBT_RIGHT : MBT_MIDDLE));
 
-		CefMouseEvent Event = GetCefMouseEvent(MyGeometry, MouseEvent, bIsPopup);
-		InternalCefBrowser->GetHost()->SendMouseClickEvent(Event, Type, false,1);
+			CefMouseEvent Event = GetCefMouseEvent(MyGeometry, MouseEvent, bIsPopup);
+
+			// If the click happened inside a drag region we enable window dragging which will start firing OnDragWindow events on mouse move
+			if (Type == MBT_LEFT && IsInDragRegion(FIntPoint(Event.x, Event.y)))
+			{
+				bDraggingWindow = true;
+			}
+
+			InternalCefBrowser->GetHost()->SendMouseClickEvent(Event, Type, false,1);
 			Reply = FReply::Handled();
 		}
 	}
@@ -1175,12 +1237,17 @@ FReply FCEFWebBrowserWindow::OnMouseButtonUp(const FGeometry& MyGeometry, const
 
 		if(bIsCefSupportedButton)
 		{
-		CefBrowserHost::MouseButtonType Type =
-			(Button == EKeys::LeftMouseButton ? MBT_LEFT : (
-			Button == EKeys::RightMouseButton ? MBT_RIGHT : MBT_MIDDLE));
+			CefBrowserHost::MouseButtonType Type =
+				(Button == EKeys::LeftMouseButton ? MBT_LEFT : (
+				Button == EKeys::RightMouseButton ? MBT_RIGHT : MBT_MIDDLE));
 
-		CefMouseEvent Event = GetCefMouseEvent(MyGeometry, MouseEvent, bIsPopup);
-		InternalCefBrowser->GetHost()->SendMouseClickEvent(Event, Type, true, 1);
+			if (Type == MBT_LEFT)
+			{
+				bDraggingWindow = false;
+			}
+
+			CefMouseEvent Event = GetCefMouseEvent(MyGeometry, MouseEvent, bIsPopup);
+			InternalCefBrowser->GetHost()->SendMouseClickEvent(Event, Type, true, 1);
 			Reply = FReply::Handled();
 		}
 		else if(Button == EKeys::ThumbMouseButton && bThumbMouseButtonNavigation)
@@ -1216,14 +1283,14 @@ FReply FCEFWebBrowserWindow::OnMouseButtonDoubleClick(const FGeometry& MyGeometr
 
 		if(bIsCefSupportedButton)
 		{
-		CefBrowserHost::MouseButtonType Type =
-			(Button == EKeys::LeftMouseButton ? MBT_LEFT : (
-			Button == EKeys::RightMouseButton ? MBT_RIGHT : MBT_MIDDLE));
+			CefBrowserHost::MouseButtonType Type =
+				(Button == EKeys::LeftMouseButton ? MBT_LEFT : (
+				Button == EKeys::RightMouseButton ? MBT_RIGHT : MBT_MIDDLE));
 
-		CefMouseEvent Event = GetCefMouseEvent(MyGeometry, MouseEvent, bIsPopup);
-		InternalCefBrowser->GetHost()->SendMouseClickEvent(Event, Type, false, 2);
+			CefMouseEvent Event = GetCefMouseEvent(MyGeometry, MouseEvent, bIsPopup);
+			InternalCefBrowser->GetHost()->SendMouseClickEvent(Event, Type, false, 2);
 			Reply = FReply::Handled();
-	}
+		}
 	}
 	return Reply;
 }
@@ -1234,7 +1301,26 @@ FReply FCEFWebBrowserWindow::OnMouseMove(const FGeometry& MyGeometry, const FPoi
 	if (IsValid())
 	{
 		CefMouseEvent Event = GetCefMouseEvent(MyGeometry, MouseEvent, bIsPopup);
-		InternalCefBrowser->GetHost()->SendMouseMoveEvent(Event, false);
+
+		bool bEventConsumedByDragCallback = false;
+		FIntRect test;
+		if (bDraggingWindow && OnDragWindow().IsBound())
+		{
+			if (MouseEvent.IsMouseButtonDown(EKeys::LeftMouseButton))
+			{
+				bEventConsumedByDragCallback = OnDragWindow().Execute(MouseEvent);
+			}
+			else
+			{
+				bDraggingWindow = false;
+			}
+		}
+
+		if (!bEventConsumedByDragCallback)
+		{
+			InternalCefBrowser->GetHost()->SendMouseMoveEvent(Event, false);
+		}
+		
 		Reply = FReply::Handled();
 	}
 	return Reply;
@@ -1246,10 +1332,20 @@ void FCEFWebBrowserWindow::OnMouseLeave(const FPointerEvent& MouseEvent)
 	SetToolTip(CefString());
 }
 
+void FCEFWebBrowserWindow::SetSupportsMouseWheel(bool bValue)
+{
+	bSupportsMouseWheel = bValue;
+}
+
+bool FCEFWebBrowserWindow::GetSupportsMouseWheel() const
+{
+	return bSupportsMouseWheel;
+}
+
 FReply FCEFWebBrowserWindow::OnMouseWheel(const FGeometry& MyGeometry, const FPointerEvent& MouseEvent, bool bIsPopup)
 {
 	FReply Reply = FReply::Unhandled();
-	if(IsValid())
+	if(IsValid() && bSupportsMouseWheel)
 	{
 		// The original delta is reduced so this should bring it back to what CEF expects
 		const float SpinFactor = 50.0f;
@@ -1505,33 +1601,41 @@ void FCEFWebBrowserWindow::NotifyDocumentLoadingStateChange(bool IsLoading)
 			: EWebBrowserDocumentState::Completed;
 		DocumentStateChangedEvent.Broadcast(DocumentState);
 	}
-
 }
 
-void FCEFWebBrowserWindow::OnPaint(CefRenderHandler::PaintElementType Type, const CefRenderHandler::RectList& DirtyRects, const void* Buffer, int Width, int Height)
+FSlateRenderer* const FCEFWebBrowserWindow::GetRenderer()
 {
-	bool bNeedsRedraw = false;
-
-#if PLATFORM_MAC
-	// @todo: Ugly workaround for OPP-7200 and OPP-7449 until proper fix can be found.  CEF returns invalid OnPaint() buffer size on Retina display Macs, or Macs with 
-	//    HiDPI enabled, once rendering is disabled/enabled using WasHidden().  Invalidating the view or calling WasResized() after enabling rendering is 
-	//    not sufficient.  For the current workaround, we must dirty the viewport size and call WasResized().
-	if (FIntPoint(Width, Height) == (ViewportSize * 2))
+	if (FSlateApplication::IsInitialized())
 	{
-		ViewportSize.Y += 1;
-		InternalCefBrowser->GetHost()->WasResized();
-
-		// We ignore this frame.
-		return;
+		if (FSlateRenderer* Renderer = FSlateApplication::Get().GetRenderer())
+		{
+			if (!Renderer->HasLostDevice())
+			{
+				return Renderer;
+			}
+		}
 	}
-#endif
 
+	ReleaseTextures();
+	return nullptr;
+}
+
+void FCEFWebBrowserWindow::HandleRenderingError()
+{
+	// GetRenderer handles errors already
+	GetRenderer();
+}
 
-	if (UpdatableTextures[Type] == nullptr && FSlateApplication::IsInitialized())
+void FCEFWebBrowserWindow::OnPaint(CefRenderHandler::PaintElementType Type, const CefRenderHandler::RectList& DirtyRects, const void* Buffer, int Width, int Height)
+{
+	bool bNeedsRedraw = false;
+
+	if (UpdatableTextures[Type] == nullptr)
 	{
-		if (FSlateRenderer* Renderer = FSlateApplication::Get().GetRenderer())
+		if (FSlateRenderer* const Renderer = GetRenderer())
 		{
 			UpdatableTextures[Type] = Renderer->CreateUpdatableTexture(Width, Height);
+			HandleRenderingError();
 		}
 	}
 
@@ -1549,12 +1653,16 @@ void FCEFWebBrowserWindow::OnPaint(CefRenderHandler::PaintElementType Type, cons
 		else
 		{
 			UpdatableTextures[Type]->UpdateTextureThreadSafeRaw(Width, Height, Buffer, Dirty);
+			HandleRenderingError();
 
 		    if (Type == PET_POPUP && bShowPopupRequested)
 		    {
 			    bShowPopupRequested = false;
 			    bPopupHasFocus = true;
-			    FIntPoint PopupSize = FIntPoint(Width, Height);
+
+				const float DPIScale = FPlatformApplicationMisc::GetDPIScaleFactorAtPoint(PopupPosition.X, PopupPosition.Y);
+			    FIntPoint PopupSize = FIntPoint(Width / DPIScale, Height / DPIScale);
+
 			    FIntRect PopupRect = FIntRect(PopupPosition, PopupPosition + PopupSize);
 			    OnShowPopup().Broadcast(PopupRect);
 		    }
@@ -1579,6 +1687,7 @@ void FCEFWebBrowserWindow::UpdateVideoBuffering()
 		if (SlateTextureData != nullptr )
 		{
 			UpdatableTextures[PET_VIEW]->UpdateTextureThreadSafeWithTextureData(SlateTextureData);
+			HandleRenderingError();
 		}
 	}
 }
@@ -1799,7 +1908,8 @@ int32 FCEFWebBrowserWindow::GetCefMouseModifiers(const FPointerEvent& InMouseEve
 CefMouseEvent FCEFWebBrowserWindow::GetCefMouseEvent(const FGeometry& MyGeometry, const FPointerEvent& MouseEvent, bool bIsPopup)
 {
 	CefMouseEvent Event;
-	FVector2D LocalPos = MyGeometry.AbsoluteToLocal(MouseEvent.GetScreenSpacePosition()) * MyGeometry.Scale;
+	const float DPIScale = MyGeometry.Scale / (ParentWindow.IsValid() ? ParentWindow->GetNativeWindow()->GetDPIScaleFactor() : 1.0f);
+	FVector2D LocalPos = MyGeometry.AbsoluteToLocal(MouseEvent.GetScreenSpacePosition()) * DPIScale;
 	if (bIsPopup)
 	{
 		LocalPos += PopupPosition;
@@ -2078,4 +2188,24 @@ void FCEFWebBrowserWindow::OnImeCompositionRangeChanged(
 }
 #endif
 
+void FCEFWebBrowserWindow::UpdateDragRegions(const TArray<FWebBrowserDragRegion>& Regions)
+{
+	DragRegions = Regions;
+}
+
+bool FCEFWebBrowserWindow::IsInDragRegion(const FIntPoint& Point)
+{
+	// Here we traverse the array of drag regions backwards because we assume the drag regions are z ordered such that 
+	//    the end of the list contains the drag regions of the top most elements of the web page.  We can stop checking
+	//    once we hit a region that contains our point.
+	for (int32 Idx = DragRegions.Num() - 1; Idx >= 0; --Idx)
+	{
+		if (DragRegions[Idx].Rect.Contains(Point))
+		{
+			return DragRegions[Idx].bDraggable;
+		}
+	}
+	return false;
+}
+
 #endif
diff --git a/Engine/Source/Runtime/WebBrowser/Private/CEF/CEFWebBrowserWindow.h b/Engine/Source/Runtime/WebBrowser/Private/CEF/CEFWebBrowserWindow.h
--- a/Engine/Source/Runtime/WebBrowser/Private/CEF/CEFWebBrowserWindow.h
+++ b/Engine/Source/Runtime/WebBrowser/Private/CEF/CEFWebBrowserWindow.h
@@ -74,6 +74,20 @@ struct FWebBrowserWindowInfo
 	CefRefPtr<FCEFBrowserHandler> Handler;
 };
 
+/**
+* Representation of a window drag region.
+*/
+struct FWebBrowserDragRegion
+{
+	FWebBrowserDragRegion(const FIntRect& InRect, bool bInDraggable)
+		: Rect(InRect)
+		, bDraggable(bInDraggable)
+	{}
+
+	FIntRect Rect;
+	bool bDraggable;
+};
+
 /**
  * Implementation of interface for dealing with a Web Browser window.
  */
@@ -141,7 +155,8 @@ public:
 	virtual FReply OnMouseButtonDoubleClick(const FGeometry& MyGeometry, const FPointerEvent& MouseEvent, bool bIsPopup) override;
 	virtual FReply OnMouseMove(const FGeometry& MyGeometry, const FPointerEvent& MouseEvent, bool bIsPopup) override;
 	virtual void OnMouseLeave(const FPointerEvent& MouseEvent) override;
-
+	virtual void SetSupportsMouseWheel(bool bValue) override;
+	virtual bool GetSupportsMouseWheel() const override;
 	virtual FReply OnMouseWheel(const FGeometry& MyGeometry, const FPointerEvent& MouseEvent, bool bIsPopup) override;
 	virtual void OnFocus(bool SetFocus, bool bIsPopup) override;
 	virtual void OnCaptureLost() override;
@@ -250,6 +265,11 @@ public:
 		return SuppressContextMenuDelgate;
 	}
 
+	virtual FOnDragWindow& OnDragWindow() override
+	{
+		return DragWindowDelegate;
+	}
+
 private:
 
 	/**
@@ -423,6 +443,8 @@ private:
 		const CefRenderHandler::RectList& CharacterBounds);
 #endif
 
+	void UpdateDragRegions(const TArray<FWebBrowserDragRegion>& Regions);
+
 public:
 
 	/**
@@ -474,6 +496,18 @@ public:
 
 private:
 
+	/** @return the currently valid renderer, if available */
+	FSlateRenderer* const GetRenderer();
+
+	/**Checks whether an error with the renderer occurred and handles it if one has */
+	void HandleRenderingError();
+
+	/** Releases the updatable textures */
+	void ReleaseTextures();
+
+	/** Creates the initial updatable textures */
+	bool CreateInitialTextures();
+
 	/** Executes or defers a LoadUrl navigation */
 	void RequestNavigationInternal(FString Url, FString Contents);
 
@@ -492,6 +526,9 @@ private:
 	/** Used to convert a FPointerEvent to a CefMouseEvent */
 	CefMouseEvent GetCefMouseEvent(const FGeometry& MyGeometry, const FPointerEvent& MouseEvent, bool bIsPopup);
 
+	/** Specifies whether or not a point falls within any tagged drag regions that are draggable. */
+	bool IsInDragRegion(const FIntPoint& Point);
+
 private:
 
 	/** Current state of the document being loaded. */
@@ -518,6 +555,9 @@ private:
 	/** Current size of this window. */
 	FIntPoint ViewportSize;
 
+	/** Current DPI scale factor of this window. */
+	float ViewportDPIScaleFactor;
+
 	/** Whether this window is closing. */
 	bool bIsClosing;
 
@@ -575,6 +615,8 @@ private:
 	/** Delegate for suppressing context menu */
 	FOnSuppressContextMenu SuppressContextMenuDelgate;
 
+	/** Delegate that is executed when a drag event is detected in an area of the web page tagged as a drag region. */
+	FOnDragWindow DragWindowDelegate;
 
 	/** Tracks the current mouse cursor */
 	EMouseCursor::Type Cursor;
@@ -591,6 +633,9 @@ private:
 	/** Tracks whether the widget has been resized and needs to be refreshed */
 	bool bNeedsResize;
 
+	/** Tracks whether or not the user initiated a window drag by clicking on a page's drag region. */
+	bool bDraggingWindow;
+
 	/** Used for unhandled key events forwarding*/
 	TOptional<FKeyEvent> PreviousKeyDownEvent;
 	TOptional<FKeyEvent> PreviousKeyUpEvent;
@@ -603,6 +648,8 @@ private:
 	bool bMainHasFocus;
 	bool bPopupHasFocus;
 
+	bool bSupportsMouseWheel;
+
 	FIntPoint PopupPosition;
 	bool bShowPopupRequested;
 
@@ -630,6 +677,8 @@ private:
 	TSharedPtr<FCEFImeHandler> Ime;
 #endif
 
+	TArray<FWebBrowserDragRegion> DragRegions;
+
 	TSharedPtr<SWindow> ParentWindow;
 };
 
diff --git a/Engine/Source/Runtime/WebBrowser/Private/IOS/IOSPlatformWebBrowser.cpp b/Engine/Source/Runtime/WebBrowser/Private/IOS/IOSPlatformWebBrowser.cpp
--- a/Engine/Source/Runtime/WebBrowser/Private/IOS/IOSPlatformWebBrowser.cpp
+++ b/Engine/Source/Runtime/WebBrowser/Private/IOS/IOSPlatformWebBrowser.cpp
@@ -674,6 +674,16 @@ void FWebBrowserWindow::OnMouseLeave(const FPointerEvent& MouseEvent)
 {
 }
 
+void FWebBrowserWindow::SetSupportsMouseWheel(bool bValue)
+{
+
+}
+
+bool FWebBrowserWindow::GetSupportsMouseWheel() const
+{
+	return false;
+}
+
 FReply FWebBrowserWindow::OnMouseWheel(const FGeometry& MyGeometry, const FPointerEvent& MouseEvent, bool bIsPopup)
 {
 	return FReply::Unhandled();
diff --git a/Engine/Source/Runtime/WebBrowser/Private/IOS/IOSPlatformWebBrowser.h b/Engine/Source/Runtime/WebBrowser/Private/IOS/IOSPlatformWebBrowser.h
--- a/Engine/Source/Runtime/WebBrowser/Private/IOS/IOSPlatformWebBrowser.h
+++ b/Engine/Source/Runtime/WebBrowser/Private/IOS/IOSPlatformWebBrowser.h
@@ -123,6 +123,8 @@ public:
 	virtual FReply OnMouseButtonDoubleClick(const FGeometry& MyGeometry, const FPointerEvent& MouseEvent, bool bIsPopup) override;
 	virtual FReply OnMouseMove(const FGeometry& MyGeometry, const FPointerEvent& MouseEvent, bool bIsPopup) override;
 	virtual void OnMouseLeave(const FPointerEvent& MouseEvent) override;
+	virtual void SetSupportsMouseWheel(bool bValue) override;
+	virtual bool GetSupportsMouseWheel() const override;
 	virtual FReply OnMouseWheel(const FGeometry& MyGeometry, const FPointerEvent& MouseEvent, bool bIsPopup) override;
 	virtual void OnFocus(bool SetFocus, bool bIsPopup) override;
 	virtual void OnCaptureLost() override;
@@ -240,6 +242,10 @@ public:
 		return SuppressContextMenuDelgate;
 	}
 
+	virtual FOnDragWindow& OnDragWindow() override
+	{
+		return DragWindowDelegate;
+	}
 
 public:
 
@@ -303,6 +309,9 @@ private:
 	/** Delegate for suppressing context menu */
 	FOnSuppressContextMenu SuppressContextMenuDelgate;
 
+	/** Delegate that is executed when a drag event is detected in an area of the web page tagged as a drag region. */
+	FOnDragWindow DragWindowDelegate;
+
 	TSharedPtr<SWindow> ParentWindow;
 
 	FIntPoint IOSWindowSize;
diff --git a/Engine/Source/Runtime/WebBrowser/Private/SWebBrowser.cpp b/Engine/Source/Runtime/WebBrowser/Private/SWebBrowser.cpp
--- a/Engine/Source/Runtime/WebBrowser/Private/SWebBrowser.cpp
+++ b/Engine/Source/Runtime/WebBrowser/Private/SWebBrowser.cpp
@@ -125,6 +125,7 @@ void SWebBrowser::Construct(const FArguments& InArgs, const TSharedPtr<IWebBrows
 				.OnDismissAllDialogs(OnDismissAllDialogs)
 				.Visibility(this, &SWebBrowser::GetViewportVisibility)
 				.OnSuppressContextMenu(InArgs._OnSuppressContextMenu)
+				.OnDragWindow(InArgs._OnDragWindow)
 			]
 			+ SOverlay::Slot()
 			.HAlign(HAlign_Center)
diff --git a/Engine/Source/Runtime/WebBrowser/Private/SWebBrowserView.cpp b/Engine/Source/Runtime/WebBrowser/Private/SWebBrowserView.cpp
--- a/Engine/Source/Runtime/WebBrowser/Private/SWebBrowserView.cpp
+++ b/Engine/Source/Runtime/WebBrowser/Private/SWebBrowserView.cpp
@@ -63,16 +63,15 @@ SWebBrowserView::~SWebBrowserView()
 		}
 	}
 
-	if (SlateParentWindowSetupTickHandle.IsValid())
+	TSharedPtr<SWindow> SlateParentWindow = SlateParentWindowPtr.Pin();
+	if (SlateParentWindow.IsValid())
 	{
-		FTicker::GetCoreTicker().RemoveTicker(SlateParentWindowSetupTickHandle);
-		SlateParentWindowSetupTickHandle.Reset();
+		SlateParentWindow->GetOnWindowDeactivatedEvent().RemoveAll(this);
 	}
 
-	if (SlateParentWindowPtr.IsValid())
+	if (SlateParentWindow.IsValid())
 	{
-		SlateParentWindowPtr.Pin()->GetOnWindowDeactivatedEvent().RemoveAll(this);
-		SlateParentWindowPtr.Pin()->GetOnWindowActivatedEvent().RemoveAll(this);
+		SlateParentWindow->GetOnWindowActivatedEvent().RemoveAll(this);
 	}
 }
 
@@ -108,6 +107,7 @@ void SWebBrowserView::Construct(const FArguments& InArgs, const TSharedPtr<IWebB
 			Settings.bShowErrorMessage = InArgs._ShowErrorMessage;
 			Settings.BackgroundColor = InArgs._BackgroundColor;
 			Settings.Context = InArgs._ContextSettings;
+			Settings.AltRetryDomains = InArgs._AltRetryDomains;
 
 			BrowserWindow = IWebBrowserModule::Get().GetSingleton()->CreateBrowserWindow(Settings);
 		}
@@ -142,6 +142,7 @@ void SWebBrowserView::Construct(const FArguments& InArgs, const TSharedPtr<IWebB
 		BrowserWindow->OnTitleChanged().AddSP(this, &SWebBrowserView::HandleTitleChanged);
 		BrowserWindow->OnUrlChanged().AddSP(this, &SWebBrowserView::HandleUrlChanged);
 		BrowserWindow->OnToolTip().AddSP(this, &SWebBrowserView::HandleToolTip);
+		OnCreateToolTip = InArgs._OnCreateToolTip;
 
 		if (!BrowserWindow->OnBeforeBrowse().IsBound())
 		{
@@ -178,24 +179,15 @@ void SWebBrowserView::Construct(const FArguments& InArgs, const TSharedPtr<IWebB
 		BrowserWindow->OnSuppressContextMenu().BindSP(this, &SWebBrowserView::HandleSuppressContextMenu);
 		OnSuppressContextMenu = InArgs._OnSuppressContextMenu;
 
+		BrowserWindow->OnDragWindow().BindSP(this, &SWebBrowserView::HandleDrag);
+		OnDragWindow = InArgs._OnDragWindow;
+
 		BrowserViewport = MakeShareable(new FWebBrowserViewport(BrowserWindow));
 #if WITH_CEF3
 		BrowserWidget->SetViewportInterface(BrowserViewport.ToSharedRef());
 #endif
-		SetupParentWindowHandlers();
 		// If we could not obtain the parent window during widget construction, we'll defer and keep trying.
-		if (!SlateParentWindowPtr.IsValid())
-		{
-			SlateParentWindowSetupTickHandle = FTicker::GetCoreTicker().AddTicker(FTickerDelegate::CreateLambda([this](float) -> bool
-			{
-                QUICK_SCOPE_CYCLE_COUNTER(STAT_SWebBrowserView_SetupParentWindowHandlers_LambdaTick);
-				this->SetupParentWindowHandlers();
-				bool ContinueTick = !SlateParentWindowPtr.IsValid();
-				return ContinueTick;
-			}));
-		}
-
-		BrowserWindow->SetParentWindow(InArgs._ParentWindow);
+		SetupParentWindowHandlers();
 	}
 	else
 	{
@@ -205,6 +197,12 @@ void SWebBrowserView::Construct(const FArguments& InArgs, const TSharedPtr<IWebB
 
 int32 SWebBrowserView::OnPaint(const FPaintArgs& Args, const FGeometry& AllottedGeometry, const FSlateRect& MyCullingRect, FSlateWindowElementList& OutDrawElements, int32 LayerId, const FWidgetStyle& InWidgetStyle, bool bParentEnabled) const
 {
+	if (!SlateParentWindowPtr.IsValid())
+	{
+		SWebBrowserView* MutableThis = const_cast<SWebBrowserView*>(this);
+		MutableThis->SetupParentWindowHandlers();
+	}
+
 	int32 Layer = SCompoundWidget::OnPaint(Args, AllottedGeometry, MyCullingRect, OutDrawElements, LayerId, InWidgetStyle, bParentEnabled);
 	
 	// Cache a reference to our parent window, if we didn't already reference it.
@@ -363,10 +361,25 @@ bool SWebBrowserView::IsInitialized() const
 
 void SWebBrowserView::SetupParentWindowHandlers()
 {
-	if (SlateParentWindowPtr.IsValid() && BrowserWindow.IsValid())
+	if (!SlateParentWindowPtr.IsValid())
 	{
-		SlateParentWindowPtr.Pin()->GetOnWindowDeactivatedEvent().AddSP(this, &SWebBrowserView::HandleWindowDeactivated);
-		SlateParentWindowPtr.Pin()->GetOnWindowActivatedEvent().AddSP(this, &SWebBrowserView::HandleWindowActivated);
+		SlateParentWindowPtr = FSlateApplication::Get().FindWidgetWindow(SharedThis(this));
+
+		TSharedPtr<SWindow> SlateParentWindow = SlateParentWindowPtr.Pin();
+		if (SlateParentWindow.IsValid() && BrowserWindow.IsValid())
+		{
+			if (!SlateParentWindow->GetOnWindowDeactivatedEvent().IsBoundToObject(this))
+			{
+				SlateParentWindow->GetOnWindowDeactivatedEvent().AddSP(this, &SWebBrowserView::HandleWindowDeactivated);
+			}
+
+			if (!SlateParentWindow->GetOnWindowActivatedEvent().IsBoundToObject(this))
+			{
+				SlateParentWindow->GetOnWindowActivatedEvent().AddSP(this, &SWebBrowserView::HandleWindowActivated);
+			}
+
+			BrowserWindow->SetParentWindow(SlateParentWindow);
+		}
 	}
 }
 
@@ -426,6 +439,11 @@ void SWebBrowserView::HandleToolTip(FString ToolTipText)
 		FSlateApplication::Get().CloseToolTip();
 		SetToolTip(nullptr);
 	}
+	else if (OnCreateToolTip.IsBound())
+	{
+		SetToolTip(OnCreateToolTip.Execute(FText::FromString(ToolTipText)));
+		FSlateApplication::Get().UpdateToolTip(true);
+	}
 	else
 	{
 		SetToolTipText(FText::FromString(ToolTipText));
@@ -615,4 +633,13 @@ bool SWebBrowserView::HandleSuppressContextMenu()
 	return false;
 }
 
+bool SWebBrowserView::HandleDrag(const FPointerEvent& MouseEvent)
+{
+	if (OnDragWindow.IsBound())
+	{
+		return OnDragWindow.Execute(MouseEvent);
+	}
+	return false;
+}
+
 #undef LOCTEXT_NAMESPACE
diff --git a/Engine/Source/Runtime/WebBrowser/Private/WebBrowserSingleton.cpp b/Engine/Source/Runtime/WebBrowser/Private/WebBrowserSingleton.cpp
--- a/Engine/Source/Runtime/WebBrowser/Private/WebBrowserSingleton.cpp
+++ b/Engine/Source/Runtime/WebBrowser/Private/WebBrowserSingleton.cpp
@@ -240,6 +240,9 @@ FWebBrowserSingleton::FWebBrowserSingleton(const FWebBrowserInitSettings& WebBro
 	CefSettings Settings;
 	Settings.no_sandbox = true;
 	Settings.command_line_args_disabled = true;
+#if !PLATFORM_LINUX
+	Settings.enable_net_security_expiration = true;
+#endif
 
 	FString CefLogFile(FPaths::Combine(*FPaths::ProjectLogDir(), TEXT("cef3.log")));
 	CefLogFile = FPaths::ConvertRelativePathToFull(CefLogFile);
@@ -384,7 +387,6 @@ TSharedPtr<IWebBrowserWindow> FWebBrowserSingleton::CreateBrowserWindow(
 	FString InitialURL = BrowserWindowInfo->Browser->GetMainFrame()->GetURL().ToWString().c_str();
 	TSharedPtr<FCEFWebBrowserWindow> NewBrowserWindow(new FCEFWebBrowserWindow(BrowserWindowInfo->Browser, BrowserWindowInfo->Handler, InitialURL, ContentsToLoad, bShowErrorMessage, bThumbMouseButtonNavigation, bUseTransparency, bJSBindingsToLoweringEnabled));
 	BrowserWindowInfo->Handler->SetBrowserWindow(NewBrowserWindow);
-
 	{
 		FScopeLock Lock(&WindowInterfacesCS);
 		WindowInterfaces.Add(NewBrowserWindow);
@@ -402,7 +404,8 @@ TSharedPtr<IWebBrowserWindow> FWebBrowserSingleton::CreateBrowserWindow(
 	TOptional<FString> ContentsToLoad,
 	bool ShowErrorMessage,
 	FColor BackgroundColor,
-	int BrowserFrameRate )
+	int BrowserFrameRate,
+	const TArray<FString>& AltRetryDomains)
 {
 	FCreateBrowserWindowSettings Settings;
 	Settings.OSWindowHandle = OSWindowHandle;
@@ -413,6 +416,7 @@ TSharedPtr<IWebBrowserWindow> FWebBrowserSingleton::CreateBrowserWindow(
 	Settings.bShowErrorMessage = ShowErrorMessage;
 	Settings.BackgroundColor = BackgroundColor;
 	Settings.BrowserFrameRate = BrowserFrameRate;
+	Settings.AltRetryDomains = AltRetryDomains;
 
 	return CreateBrowserWindow(Settings);
 }
@@ -457,7 +461,7 @@ TSharedPtr<IWebBrowserWindow> FWebBrowserSingleton::CreateBrowserWindow(const FC
 
 
 		// WebBrowserHandler implements browser-level callbacks.
-		CefRefPtr<FCEFBrowserHandler> NewHandler(new FCEFBrowserHandler(WindowSettings.bUseTransparency));
+		CefRefPtr<FCEFBrowserHandler> NewHandler(new FCEFBrowserHandler(WindowSettings.bUseTransparency, WindowSettings.AltRetryDomains));
 
 		CefRefPtr<CefRequestContext> RequestContext = nullptr;
 		if (WindowSettings.Context.IsSet())
@@ -472,6 +476,9 @@ TSharedPtr<IWebBrowserWindow> FWebBrowserSingleton::CreateBrowserWindow(const FC
 				CefString(&RequestContextSettings.cache_path) = *Context.CookieStorageLocation;
 				RequestContextSettings.persist_session_cookies = Context.bPersistSessionCookies;
 				RequestContextSettings.ignore_certificate_errors = Context.bIgnoreCertificateErrors;
+#if !PLATFORM_LINUX
+				RequestContextSettings.enable_net_security_expiration = Context.bEnableNetSecurityExpiration;
+#endif
 
 				//Create a new one
 				RequestContext = CefRequestContext::CreateContext(RequestContextSettings, nullptr);
@@ -499,11 +506,11 @@ TSharedPtr<IWebBrowserWindow> FWebBrowserSingleton::CreateBrowserWindow(const FC
 				WindowSettings.bUseTransparency,
 				bJSBindingsToLoweringEnabled));
 			NewHandler->SetBrowserWindow(NewBrowserWindow);
-
 			{
 				FScopeLock Lock(&WindowInterfacesCS);
 				WindowInterfaces.Add(NewBrowserWindow);
 			}
+			
 			return NewBrowserWindow;
 		}
 	}
@@ -517,7 +524,6 @@ TSharedPtr<IWebBrowserWindow> FWebBrowserSingleton::CreateBrowserWindow(const FC
 		WindowSettings.bUseTransparency,
 		bJSBindingsToLoweringEnabled));
 
-	//WindowInterfaces.Add(NewBrowserWindow);
 	return NewBrowserWindow;
 #elif PLATFORM_IOS || PLATFORM_PS4
 	// Create new window
@@ -528,7 +534,6 @@ TSharedPtr<IWebBrowserWindow> FWebBrowserSingleton::CreateBrowserWindow(const FC
 		WindowSettings.bThumbMouseButtonNavigation, 
 		WindowSettings.bUseTransparency));
 
-	//WindowInterfaces.Add(NewBrowserWindow);
 	return NewBrowserWindow;
 #endif
 	return nullptr;
@@ -539,25 +544,28 @@ bool FWebBrowserSingleton::Tick(float DeltaTime)
     QUICK_SCOPE_CYCLE_COUNTER(STAT_FWebBrowserSingleton_Tick);
 
 #if WITH_CEF3
-	FScopeLock Lock(&WindowInterfacesCS);
-	bool bIsSlateAwake = FSlateApplication::IsInitialized() && !FSlateApplication::Get().IsSlateAsleep();
-	// Remove any windows that have been deleted and check whether it's currently visible
-	for (int32 Index = WindowInterfaces.Num() - 1; Index >= 0; --Index)
 	{
-		if (!WindowInterfaces[Index].IsValid())
-		{
-			WindowInterfaces.RemoveAt(Index);
-		}
-		else if (bIsSlateAwake) // only check for Tick activity if Slate is currently ticking
+		FScopeLock Lock(&WindowInterfacesCS);
+		bool bIsSlateAwake = FSlateApplication::IsInitialized() && !FSlateApplication::Get().IsSlateAsleep();
+		// Remove any windows that have been deleted and check whether it's currently visible
+		for (int32 Index = WindowInterfaces.Num() - 1; Index >= 0; --Index)
 		{
-			TSharedPtr<FCEFWebBrowserWindow> BrowserWindow = WindowInterfaces[Index].Pin();
-			if(BrowserWindow.IsValid())
+			if (!WindowInterfaces[Index].IsValid())
 			{
-				// Test if we've ticked recently. If not assume the browser window has become hidden.
-				BrowserWindow->CheckTickActivity();
+				WindowInterfaces.RemoveAt(Index);
+			}
+			else if (bIsSlateAwake) // only check for Tick activity if Slate is currently ticking
+			{
+				TSharedPtr<FCEFWebBrowserWindow> BrowserWindow = WindowInterfaces[Index].Pin();
+				if(BrowserWindow.IsValid())
+				{
+					// Test if we've ticked recently. If not assume the browser window has become hidden.
+					BrowserWindow->CheckTickActivity();
+				}
 			}
 		}
 	}
+
 	CefDoMessageLoopWork();
 
 	// Update video buffering for any windows that need it
@@ -637,6 +645,9 @@ bool FWebBrowserSingleton::RegisterContext(const FBrowserContextSettings& Settin
 	CefString(&RequestContextSettings.cache_path) = *Settings.CookieStorageLocation;
 	RequestContextSettings.persist_session_cookies = Settings.bPersistSessionCookies;
 	RequestContextSettings.ignore_certificate_errors = Settings.bIgnoreCertificateErrors;
+#if !PLATFORM_LINUX
+	RequestContextSettings.enable_net_security_expiration = Settings.bEnableNetSecurityExpiration;
+#endif
 
 	//Create a new one
 	CefRefPtr<CefRequestContext> RequestContext = CefRequestContext::CreateContext(RequestContextSettings, nullptr);
diff --git a/Engine/Source/Runtime/WebBrowser/Private/WebBrowserSingleton.h b/Engine/Source/Runtime/WebBrowser/Private/WebBrowserSingleton.h
--- a/Engine/Source/Runtime/WebBrowser/Private/WebBrowserSingleton.h
+++ b/Engine/Source/Runtime/WebBrowser/Private/WebBrowserSingleton.h
@@ -84,7 +84,8 @@ public:
 		TOptional<FString> ContentsToLoad = TOptional<FString>(),
 		bool ShowErrorMessage = true,
 		FColor BackgroundColor = FColor(255, 255, 255, 255),
-		int BrowserFrameRate = 24 ) override;
+		int BrowserFrameRate = 24,
+		const TArray<FString>& AltRetryDomains = TArray<FString>()) override;
 
 	TSharedPtr<IWebBrowserWindow> CreateBrowserWindow(const FCreateBrowserWindowSettings& Settings) override;
 
diff --git a/Engine/Source/Runtime/WebBrowser/Private/WebBrowserViewport.cpp b/Engine/Source/Runtime/WebBrowser/Private/WebBrowserViewport.cpp
--- a/Engine/Source/Runtime/WebBrowser/Private/WebBrowserViewport.cpp
+++ b/Engine/Source/Runtime/WebBrowser/Private/WebBrowserViewport.cpp
@@ -26,7 +26,8 @@ void FWebBrowserViewport::Tick( const FGeometry& AllottedGeometry, double InCurr
 {
 	if (!bIsPopup)
 	{
-		FVector2D AbsoluteSize = AllottedGeometry.GetLocalSize() * AllottedGeometry.Scale;
+		const float DPIScale = AllottedGeometry.Scale / (WebBrowserWindow->GetParentWindow().IsValid() ? WebBrowserWindow->GetParentWindow()->GetNativeWindow()->GetDPIScaleFactor() : 1.0f);
+		FVector2D AbsoluteSize = AllottedGeometry.GetLocalSize() * DPIScale;
 		WebBrowserWindow->SetViewportSize(AbsoluteSize.IntPoint());
 
 #if WITH_CEF3
diff --git a/Engine/Source/Runtime/WebBrowser/Public/IWebBrowserCookieManager.h b/Engine/Source/Runtime/WebBrowser/Public/IWebBrowserCookieManager.h
--- a/Engine/Source/Runtime/WebBrowser/Public/IWebBrowserCookieManager.h
+++ b/Engine/Source/Runtime/WebBrowser/Public/IWebBrowserCookieManager.h
@@ -3,41 +3,41 @@
 
 #include "CoreMinimal.h"
 
-class IWebBrowserCookieManager
+struct FWebBrowserCookie
 {
 public:
+	// The cookie name.
+	FString Name;
 
-	struct FCookie
-	{
-	public:
-		// The cookie name.
-		FString Name;
-
-		// The cookie value.
-		FString Value;
+	// The cookie value.
+	FString Value;
 
-		// If is empty a host cookie will be created instead of a domain
-		// cookie. Domain cookies are stored with a leading "." and are visible to
-		// sub-domains whereas host cookies are not.
-		FString Domain;
+	// If is empty a host cookie will be created instead of a domain
+	// cookie. Domain cookies are stored with a leading "." and are visible to
+	// sub-domains whereas host cookies are not.
+	FString Domain;
 
-		// If is non-empty only URLs at or below the path will get the cookie
-		// value.
-		FString Path;
+	// If is non-empty only URLs at or below the path will get the cookie
+	// value.
+	FString Path;
 
-		// If true the cookie will only be sent for HTTPS requests.
-		bool bSecure;
+	// If true the cookie will only be sent for HTTPS requests.
+	bool bSecure;
 
-		// If true the cookie will only be sent for HTTP requests.
-		bool bHttpOnly;
+	// If true the cookie will only be sent for HTTP requests.
+	bool bHttpOnly;
 
-		// If true the cookie will expire at the specified Expires datetime.
-		bool bHasExpires;
+	// If true the cookie will expire at the specified Expires datetime.
+	bool bHasExpires;
 
-		// The cookie expiration date is only valid if bHasExpires is true.
-		FDateTime Expires;
-	};
+	// The cookie expiration date is only valid if bHasExpires is true.
+	FDateTime Expires;
+};
 
+class IWebBrowserCookieManager
+{
+public:
+	typedef struct FWebBrowserCookie FCookie;
 public:
 
 	/**
@@ -65,5 +65,4 @@ public:
 	 * @param Completed A callback function that will be invoked asynchronously on the game thread when the deletion is complete passing in the number of deleted objects.
 	 */
 	virtual void DeleteCookies(const FString& URL = TEXT(""), const FString& CookieName = TEXT(""), TFunction<void(int)> Completed = nullptr) = 0;
-
-};
+};
\ No newline at end of file
diff --git a/Engine/Source/Runtime/WebBrowser/Public/IWebBrowserSingleton.h b/Engine/Source/Runtime/WebBrowser/Public/IWebBrowserSingleton.h
--- a/Engine/Source/Runtime/WebBrowser/Public/IWebBrowserSingleton.h
+++ b/Engine/Source/Runtime/WebBrowser/Public/IWebBrowserSingleton.h
@@ -37,6 +37,7 @@ struct WEBBROWSER_API FBrowserContextSettings
 		, CookieStorageLocation()
 		, bPersistSessionCookies(false)
 		, bIgnoreCertificateErrors(false)
+		, bEnableNetSecurityExpiration(true)
 	{ }
 
 	FString Id;
@@ -44,6 +45,7 @@ struct WEBBROWSER_API FBrowserContextSettings
 	FString CookieStorageLocation;
 	bool bPersistSessionCookies;
 	bool bIgnoreCertificateErrors;
+	bool bEnableNetSecurityExpiration;
 };
 
 
@@ -60,6 +62,7 @@ struct WEBBROWSER_API FCreateBrowserWindowSettings
 		, BackgroundColor(FColor(255, 255, 255, 255))
 		, BrowserFrameRate(24)
 		, Context()
+		, AltRetryDomains()
 	{ }
 
 	void* OSWindowHandle;
@@ -71,6 +74,7 @@ struct WEBBROWSER_API FCreateBrowserWindowSettings
 	FColor BackgroundColor;
 	int BrowserFrameRate;
 	TOptional<FBrowserContextSettings> Context;
+	TArray<FString> AltRetryDomains;
 };
 
 /**
@@ -121,7 +125,8 @@ public:
 		TOptional<FString> ContentsToLoad = TOptional<FString>(),
 		bool ShowErrorMessage = true,
 		FColor BackgroundColor = FColor(255, 255, 255, 255),
-		int BrowserFrameRate = 24 ) = 0;
+		int BrowserFrameRate = 24,
+		const TArray<FString>& AltRetryDomains = TArray<FString>()) = 0;
 
 	virtual TSharedPtr<IWebBrowserWindow> CreateBrowserWindow(const FCreateBrowserWindowSettings& Settings) = 0;
 
diff --git a/Engine/Source/Runtime/WebBrowser/Public/IWebBrowserWindow.h b/Engine/Source/Runtime/WebBrowser/Public/IWebBrowserWindow.h
--- a/Engine/Source/Runtime/WebBrowser/Public/IWebBrowserWindow.h
+++ b/Engine/Source/Runtime/WebBrowser/Public/IWebBrowserWindow.h
@@ -175,6 +175,16 @@ public:
 	 */
 	virtual void OnMouseLeave(const FPointerEvent& MouseEvent) = 0;
 
+	/**
+	 * Sets whether mouse wheel events should be handled by the window
+	 */
+	virtual void SetSupportsMouseWheel(bool bValue) = 0;
+
+	/**
+	 * Returns whether mouse wheel events should be handled by the window
+	 */
+	virtual bool GetSupportsMouseWheel() const = 0;
+
 	/**
 	 * Called when the mouse wheel is spun
 	 *
@@ -350,6 +360,10 @@ public:
 	DECLARE_DELEGATE_RetVal(bool, FOnSuppressContextMenu);
 	virtual FOnSuppressContextMenu& OnSuppressContextMenu() = 0;
 
+	/** A delegate that is invoked when drag is detected in an area specified as a drag region on the web page. */
+	DECLARE_DELEGATE_RetVal_OneParam(bool, FOnDragWindow, const FPointerEvent& /*MouseEvent*/)
+	virtual FOnDragWindow& OnDragWindow() = 0;
+
 protected:
 
 	/** Virtual Destructor. */
diff --git a/Engine/Source/Runtime/WebBrowser/Public/SWebBrowser.h b/Engine/Source/Runtime/WebBrowser/Public/SWebBrowser.h
--- a/Engine/Source/Runtime/WebBrowser/Public/SWebBrowser.h
+++ b/Engine/Source/Runtime/WebBrowser/Public/SWebBrowser.h
@@ -25,6 +25,7 @@ public:
 	DECLARE_DELEGATE_RetVal_TwoParams(bool, FOnBeforeBrowse, const FString&, const FWebNavigationRequest& /*Request*/)
 	DECLARE_DELEGATE_RetVal_ThreeParams(bool, FOnLoadUrl, const FString& /*Method*/, const FString& /*Url*/, FString& /* Response */)
 	DECLARE_DELEGATE_RetVal_OneParam(EWebBrowserDialogEventResponse, FOnShowDialog, const TWeakPtr<IWebBrowserDialog>&);
+	DECLARE_DELEGATE_RetVal_OneParam(bool, FOnDragWindow, const FPointerEvent& /* MouseEvent */);
 
 	SLATE_BEGIN_ARGS(SWebBrowser)
 		: _InitialURL(TEXT("https://www.google.com"))
@@ -113,6 +114,9 @@ public:
 
 		SLATE_EVENT(FOnSuppressContextMenu, OnSuppressContextMenu);
 
+		/** Called when drag is detected in a web page area tagged as a drag region. */
+		SLATE_EVENT(FOnDragWindow, OnDragWindow);
+
 
 	SLATE_END_ARGS()
 
diff --git a/Engine/Source/Runtime/WebBrowser/Public/SWebBrowserView.h b/Engine/Source/Runtime/WebBrowser/Public/SWebBrowserView.h
--- a/Engine/Source/Runtime/WebBrowser/Public/SWebBrowserView.h
+++ b/Engine/Source/Runtime/WebBrowser/Public/SWebBrowserView.h
@@ -23,6 +23,7 @@ enum class EWebBrowserDocumentState;
 DECLARE_DELEGATE_RetVal_TwoParams(bool, FOnBeforePopupDelegate, FString, FString);
 DECLARE_DELEGATE_RetVal_TwoParams(bool, FOnCreateWindowDelegate, const TWeakPtr<IWebBrowserWindow>&, const TWeakPtr<IWebBrowserPopupFeatures>&);
 DECLARE_DELEGATE_RetVal_OneParam(bool, FOnCloseWindowDelegate, const TWeakPtr<IWebBrowserWindow>&);
+DECLARE_DELEGATE_RetVal_OneParam(TSharedPtr<IToolTip>, FOnCreateToolTip, const FText&);
 
 #if WITH_CEF3
 typedef SViewport SWebBrowserWidget;
@@ -38,6 +39,7 @@ public:
 	DECLARE_DELEGATE_RetVal_ThreeParams(bool, FOnLoadUrl, const FString& /*Method*/, const FString& /*Url*/, FString& /* Response */)
 	DECLARE_DELEGATE_RetVal_OneParam(EWebBrowserDialogEventResponse, FOnShowDialog, const TWeakPtr<IWebBrowserDialog>&);
 	DECLARE_DELEGATE_RetVal(bool, FOnSuppressContextMenu);
+	DECLARE_DELEGATE_RetVal_OneParam(bool, FOnDragWindow, const FPointerEvent& /* MouseEvent */);
 
 	SLATE_BEGIN_ARGS(SWebBrowserView)
 		: _InitialURL(TEXT("https://www.google.com"))
@@ -47,6 +49,7 @@ public:
 		, _BackgroundColor(255,255,255,255)
 		, _PopupMenuMethod(TOptional<EPopupMethod>())
 		, _ContextSettings()
+		, _AltRetryDomains(TArray<FString>())
 		, _ViewportSize(FVector2D::ZeroVector)
 	{ }
 
@@ -77,6 +80,9 @@ public:
 		/** Override the default global context settings for this specific window. If not set, the global default will be used. */
 		SLATE_ARGUMENT(TOptional<FBrowserContextSettings>, ContextSettings)
 
+		/** Domains to retry if original domain cannot connect. */
+		SLATE_ARGUMENT(TArray<FString>, AltRetryDomains)
+
 		/** Desired size of the web browser viewport. */
 		SLATE_ATTRIBUTE(FVector2D, ViewportSize);
 
@@ -118,6 +124,12 @@ public:
 
 		SLATE_EVENT(FOnSuppressContextMenu, OnSuppressContextMenu);
 
+		/** Called to allow overriding of ToolTip widget construction. */
+		SLATE_EVENT(FOnCreateToolTip, OnCreateToolTip)
+
+		/** Called when drag is detected in a web page area tagged as a drag region. */
+		SLATE_EVENT(FOnDragWindow, OnDragWindow)
+
 	SLATE_END_ARGS()
 
 
@@ -306,6 +318,8 @@ private:
 	void HandleWindowDeactivated();
 	void HandleWindowActivated();
 
+	bool HandleDrag(const FPointerEvent& MouseEvent);
+
 private:
 
 	/** Interface for dealing with a web browser window. */
@@ -369,11 +383,15 @@ private:
 	/** A delegate that is invoked when when the browser needs to dismiss all dialogs */
 	FSimpleDelegate OnDismissAllDialogs;
 
-	FDelegateHandle SlateParentWindowSetupTickHandle;
-
 	FOnSuppressContextMenu OnSuppressContextMenu;
 
+	/** A delegate that is invoked when when the browser wishes to create a tooltip */
+	FOnCreateToolTip OnCreateToolTip;
+
+	/** A delegate that is invoked when the browser detects drag event in within drag region */
+	FOnDragWindow OnDragWindow;
+
 protected:
 	bool HandleSuppressContextMenu();
 
-};
+};
\ No newline at end of file
diff --git a/Engine/Source/Runtime/WebBrowser/WebBrowser.Build.cs b/Engine/Source/Runtime/WebBrowser/WebBrowser.Build.cs
--- a/Engine/Source/Runtime/WebBrowser/WebBrowser.Build.cs
+++ b/Engine/Source/Runtime/WebBrowser/WebBrowser.Build.cs
@@ -20,8 +20,9 @@ public class WebBrowser : ModuleRules
 				"Slate",
 				"SlateCore",
 				"Serialization",
-			}
-		);
+                "HTTP"
+            }
+        );
 
 		if (Target.Platform == UnrealTargetPlatform.Android ||
             Target.Platform == UnrealTargetPlatform.IOS ||
