UE-50492 OnUrlChanged does not fire on mobile

#iOS
#4.21

Adding the events for load started & finished with errors

diff --git a/Engine/Source/Runtime/WebBrowser/Private/IOS/IOSPlatformWebBrowser.cpp b/Engine/Source/Runtime/WebBrowser/Private/IOS/IOSPlatformWebBrowser.cpp
--- a/Engine/Source/Runtime/WebBrowser/Private/IOS/IOSPlatformWebBrowser.cpp
+++ b/Engine/Source/Runtime/WebBrowser/Private/IOS/IOSPlatformWebBrowser.cpp
@@ -233,6 +233,19 @@ class SIOSWebBrowserWidget : public SLeafWidget
 		}
 		return Retval;
 	}
+	
+	void NotifyUrlChanged(const FString& InCurrentUrl)
+	{
+		if (WebBrowserWindowPtr.IsValid())
+		{
+			TSharedPtr<FWebBrowserWindow> BrowserWindow = WebBrowserWindowPtr.Pin();
+			if (BrowserWindow.IsValid())
+			{
+				BrowserWindow->NotifyUrlChanged(InCurrentUrl);
+			}
+		}
+
+	}
 
 	void ExecuteJavascript(const FString& Script)
 	{
@@ -548,9 +561,23 @@ shouldStartLoadWithRequest : (NSURLRequest*)InRequest
 	return YES;
 }
 
--(void)webView:(UIWebView*)InWebView didFailLoadWithError : (NSError*)InError;
+-(void)webView:(WKWebView *)webView didCommitNavigation : (WKNavigation *)navigation;
 {
+	NSString* CurrentUrl = [self.WebView URL].absoluteString;
+	NSLog(@"didCommitNavigation: %@", CurrentUrl);
+	WebBrowserWidget->NotifyUrlChanged(CurrentUrl);
+}
 
+-(void)webView:(UIWebView*)InWebView didFailLoadWithError : (NSError*)InError;
+{
+	if (InError.domain == NSURLErrorDomain && InError.code == NSURLErrorCancelled)
+	{
+		//ignore this one, interrupted load
+		return;
+	}
+	NSString* CurrentUrl = [InError.userInfo objectForKey : @"NSErrorFailingURLStringKey"]; 
+	NSLog(@"didFailLoadWithError %@", CurrentUrl);
+	WebBrowserWidget->NotifyUrlChanged(CurrentUrl);
 }
 #endif
 @end
@@ -751,4 +778,12 @@ void FWebBrowserWindow::UnbindUObject(const FString& Name, UObject* Object /*= n
 {
 }
 
+void FWebBrowserWindow::NotifyUrlChanged(const FString& InCurrentUrl)
+{
+	if (!CurrentUrl.Equals(InCurrentUrl, ESearchCase::CaseSensitive))
+	{
+		CurrentUrl = InCurrentUrl;
+		UrlChangedEvent.Broadcast(CurrentUrl);
+	}
+}
 #endif
diff --git a/Engine/Source/Runtime/WebBrowser/Private/IOS/IOSPlatformWebBrowser.h b/Engine/Source/Runtime/WebBrowser/Private/IOS/IOSPlatformWebBrowser.h
--- a/Engine/Source/Runtime/WebBrowser/Private/IOS/IOSPlatformWebBrowser.h
+++ b/Engine/Source/Runtime/WebBrowser/Private/IOS/IOSPlatformWebBrowser.h
@@ -240,7 +240,7 @@ public:
 		return SuppressContextMenuDelgate;
 	}
 
-
+	void NotifyUrlChanged(const FString& InCurrentUrl);
 public:
 
 private:
