UE-61554 Web Browser: No Transparency in 3D Widget Component on Mobile

#Android
#4.21

Modified Web Texture Material's blend mode to "Translucent"

diff --git a/Engine/Plugins/Runtime/WebBrowserWidget/Source/WebBrowserWidget/Private/WebBrowser.cpp b/Engine/Plugins/Runtime/WebBrowserWidget/Source/WebBrowserWidget/Private/WebBrowser.cpp
--- a/Engine/Plugins/Runtime/WebBrowserWidget/Source/WebBrowserWidget/Private/WebBrowser.cpp
+++ b/Engine/Plugins/Runtime/WebBrowserWidget/Source/WebBrowserWidget/Private/WebBrowser.cpp
@@ -1,4 +1,4 @@
-// Copyright 1998-2018 Epic Games, Inc. All Rights Reserved.
+﻿// Copyright 1998-2018 Epic Games, Inc. All Rights Reserved.
 
 #include "WebBrowser.h"
 #include "SWebBrowser.h"
@@ -13,16 +13,11 @@
 #include "Materials/MaterialExpressionTextureSample.h"
 #include "Materials/MaterialExpressionTextureSampleParameter2D.h"
 #include "Materials/MaterialFunction.h"
-#include "Materials/Material.h"
 #include "Factories/MaterialFactoryNew.h"
 #include "AssetRegistryModule.h"
 #include "PackageHelperFunctions.h"
 #endif
 
-#if WITH_EDITOR || PLATFORM_ANDROID
-#include "WebBrowserTexture.h"
-#endif
-
 #define LOCTEXT_NAMESPACE "WebBrowser"
 
 /////////////////////////////////////////////////////
@@ -32,22 +27,6 @@ UWebBrowser::UWebBrowser(const FObjectInitializer& ObjectInitializer)
 	: Super(ObjectInitializer)
 {
 	bIsVariable = true;
-
-#if WITH_EDITOR || PLATFORM_ANDROID
-	struct FConstructorStatics
-	{
-		ConstructorHelpers::FObjectFinder<UObject>			DefaultTextureMaterial;
-		FConstructorStatics()
-			: DefaultTextureMaterial(TEXT("/WebBrowserWidget/WebTexture_M"))
-		{}
-	};
-	static FConstructorStatics ConstructorStatics;
-
-	// Add a hard reference to UWebBrowserTexture, without this the WebBrowserTexture DLL never gets loaded on Windows.
-	UWebBrowserTexture::StaticClass();
-
-	DefaultMaterial = (UMaterial*)ConstructorStatics.DefaultTextureMaterial.Object;
-#endif
 }
 
 void UWebBrowser::LoadURL(FString NewURL)
@@ -177,11 +156,6 @@ const FText UWebBrowser::GetPaletteCategory()
 
 #endif
 
-UMaterial* UWebBrowser::GetDefaultMaterial() const
-{
-	return DefaultMaterial;
-}
-
 /////////////////////////////////////////////////////
 
 #undef LOCTEXT_NAMESPACE
diff --git a/Engine/Plugins/Runtime/WebBrowserWidget/Source/WebBrowserWidget/Private/WebBrowserWidgetModule.cpp b/Engine/Plugins/Runtime/WebBrowserWidget/Source/WebBrowserWidget/Private/WebBrowserWidgetModule.cpp
--- a/Engine/Plugins/Runtime/WebBrowserWidget/Source/WebBrowserWidget/Private/WebBrowserWidgetModule.cpp
+++ b/Engine/Plugins/Runtime/WebBrowserWidget/Source/WebBrowserWidget/Private/WebBrowserWidgetModule.cpp
@@ -2,6 +2,10 @@
 
 #include "WebBrowserWidgetModule.h"
 #include "Modules/ModuleManager.h"
+#include "WebBrowserAssetManager.h"
+#include "WebBrowserModule.h"
+#include "IWebBrowserSingleton.h"
+#include "Materials/Material.h"
 
 //////////////////////////////////////////////////////////////////////////
 // FWebBrowserWidgetModule
@@ -11,11 +15,26 @@ class FWebBrowserWidgetModule : public IWebBrowserWidgetModule
 public:
 	virtual void StartupModule() override
 	{
+		if (WebBrowserAssetMgr == nullptr)
+		{
+			WebBrowserAssetMgr = NewObject<UWebBrowserAssetManager>((UObject*)GetTransientPackage(), NAME_None, RF_Transient | RF_Public);
+			WebBrowserAssetMgr->LoadDefaultMaterials();
+
+			IWebBrowserSingleton* WebBrowserSingleton = IWebBrowserModule::Get().GetSingleton();
+			if (WebBrowserSingleton)
+			{
+				WebBrowserSingleton->SetDefaultMaterial(WebBrowserAssetMgr->GetDefaultMaterial());
+				WebBrowserSingleton->SetDefaultTranslucentMaterial(WebBrowserAssetMgr->GetDefaultTranslucentMaterial());
+			}
+		}
 	}
 
 	virtual void ShutdownModule() override
 	{
 	}
+	
+private:
+	UWebBrowserAssetManager* WebBrowserAssetMgr;
 };
 
 //////////////////////////////////////////////////////////////////////////
diff --git a/Engine/Plugins/Runtime/WebBrowserWidget/Source/WebBrowserWidget/WebBrowserWidget.build.cs b/Engine/Plugins/Runtime/WebBrowserWidget/Source/WebBrowserWidget/WebBrowserWidget.build.cs
--- a/Engine/Plugins/Runtime/WebBrowserWidget/Source/WebBrowserWidget/WebBrowserWidget.build.cs
+++ b/Engine/Plugins/Runtime/WebBrowserWidget/Source/WebBrowserWidget/WebBrowserWidget.build.cs
@@ -37,7 +37,7 @@ namespace UnrealBuildTool.Rules
                 );
             }
 			
-			if (Target.bBuildEditor || Target.Platform == UnrealTargetPlatform.Android)
+			if (Target.bBuildEditor || Target.Platform == UnrealTargetPlatform.Android || Target.Platform == UnrealTargetPlatform.IOS)
 			{
                 // WebBrowserTexture required for cooking Android
                 PrivateIncludePathModuleNames.AddRange(
diff --git a/Engine/Source/Runtime/WebBrowser/Private/Android/AndroidWebBrowserWidget.cpp b/Engine/Source/Runtime/WebBrowser/Private/Android/AndroidWebBrowserWidget.cpp
--- a/Engine/Source/Runtime/WebBrowser/Private/Android/AndroidWebBrowserWidget.cpp
+++ b/Engine/Source/Runtime/WebBrowser/Private/Android/AndroidWebBrowserWidget.cpp
@@ -20,6 +20,8 @@
 #include "Templates/SharedPointer.h"
 #include "Materials/MaterialExpressionTextureSample.h"
 #include "WebBrowserTextureSample.h"
+#include "WebBrowserModule.h"
+#include "IWebBrowserSingleton.h"
 
 // For UrlDecode
 #include "Http.h"
@@ -106,8 +108,9 @@ void SAndroidWebBrowserWidget::Construct(const FArguments& Args)
 		AllWebControls.Add(reinterpret_cast<int64>(this), StaticCastSharedRef<SAndroidWebBrowserWidget>(AsShared()));
 	}
 
-	IsAndroid3DBrowser = true;
 	WebBrowserWindowPtr = Args._WebBrowserWindow;
+	IsAndroid3DBrowser = true;
+
 	HistorySize = 0;
 	HistoryPosition = 0;
 	
@@ -131,11 +134,13 @@ void SAndroidWebBrowserWidget::Construct(const FArguments& Args)
 	}
 
 	// create wrapper material
-	UMaterial* Material = LoadObject<UMaterial>(nullptr, TEXT("/WebBrowserWidget/WebTexture_M"), nullptr, LOAD_None, nullptr);
-	if (Material)
+	IWebBrowserSingleton* WebBrowserSingleton = IWebBrowserModule::Get().GetSingleton();
+	
+	UMaterialInterface* DefaultWBMaterial = Args._UseTransparency? WebBrowserSingleton->GetDefaultTranslucentMaterial(): WebBrowserSingleton->GetDefaultMaterial();
+	if (WebBrowserSingleton && DefaultWBMaterial)
 	{
 		// create wrapper material
-		WebBrowserMaterial = UMaterialInstanceDynamic::Create(Material, nullptr);
+		WebBrowserMaterial = UMaterialInstanceDynamic::Create(DefaultWBMaterial, nullptr);
 
 		if (WebBrowserMaterial)
 		{
diff --git a/Engine/Source/Runtime/WebBrowser/Private/IOS/IOSPlatformWebBrowser.cpp b/Engine/Source/Runtime/WebBrowser/Private/IOS/IOSPlatformWebBrowser.cpp
--- a/Engine/Source/Runtime/WebBrowser/Private/IOS/IOSPlatformWebBrowser.cpp
+++ b/Engine/Source/Runtime/WebBrowser/Private/IOS/IOSPlatformWebBrowser.cpp
@@ -11,6 +11,9 @@
 #import <MetalKit/MetalKit.h>
 #include <OpenGLES/ES2/glext.h>
 #include "ExternalTexture.h"
+#include "WebBrowserModule.h"
+#include "IWebBrowserSingleton.h"
+
 
 class SIOSWebBrowserWidget : public SLeafWidget
 {
@@ -60,11 +63,13 @@ class SIOSWebBrowserWidget : public SLeafWidget
 		}
 
 		// create wrapper material
-		UMaterial* Material = LoadObject<UMaterial>(nullptr, TEXT("/WebBrowserWidget/WebTexture_M"), nullptr, LOAD_None, nullptr);
-		if (Material)
+		IWebBrowserSingleton* WebBrowserSingleton = IWebBrowserModule::Get().GetSingleton();
+
+		UMaterialInterface* DefaultWBMaterial = Args._UseTransparency ? WebBrowserSingleton->GetDefaultTranslucentMaterial() : WebBrowserSingleton->GetDefaultMaterial();
+		if (WebBrowserSingleton && DefaultWBMaterial)
 		{
 			// create wrapper material
-			WebBrowserMaterial = UMaterialInstanceDynamic::Create(Material, nullptr);
+			WebBrowserMaterial = UMaterialInstanceDynamic::Create(DefaultWBMaterial, nullptr);
 
 			if (WebBrowserMaterial)
 			{
@@ -541,11 +546,10 @@ supportsMetal : (bool)InSupportsMetal supportsMetalMRT : (bool)InSupportsMetalMR
 }
 
 #if !PLATFORM_TVOS
-- (BOOL)webView:(UIWebView*)InWebView
-shouldStartLoadWithRequest : (NSURLRequest*)InRequest
-	navigationType : (UIWebViewNavigationType)InNavigationType;
+- (void)webView:(WKWebView*)InWebView decidePolicyForNavigationAction : (WKNavigationAction*)InNavigationAction decisionHandler : (void(^)(WKNavigationActionPolicy))InDecisionHandler
 {
-	FString UrlStr([[InRequest URL]absoluteString]);
+	NSURLRequest *request = InNavigationAction.request;
+	FString UrlStr([[request URL]absoluteString]);
 
 	// Notify on the game thread
 	[FIOSAsyncTask CreateTaskWithBlock : ^ bool(void)
@@ -558,25 +562,25 @@ shouldStartLoadWithRequest : (NSURLRequest*)InRequest
 		return true;
 	}];
 
-	return YES;
+	InDecisionHandler(WKNavigationActionPolicyAllow);
 }
 
--(void)webView:(WKWebView *)webView didCommitNavigation : (WKNavigation *)navigation;
+-(void)webView:(WKWebView *)InWebView didCommitNavigation : (WKNavigation *)InNavigation
 {
 	NSString* CurrentUrl = [self.WebView URL].absoluteString;
-	NSLog(@"didCommitNavigation: %@", CurrentUrl);
+//	NSLog(@"didCommitNavigation: %@", CurrentUrl);
 	WebBrowserWidget->NotifyUrlChanged(CurrentUrl);
 }
 
--(void)webView:(UIWebView*)InWebView didFailLoadWithError : (NSError*)InError;
+-(void)webView:(WKWebView *)InWebView didFailNavigation : (WKNavigation *)InNavigation withError : (NSError*)InError
 {
 	if (InError.domain == NSURLErrorDomain && InError.code == NSURLErrorCancelled)
 	{
 		//ignore this one, interrupted load
 		return;
 	}
-	NSString* CurrentUrl = [InError.userInfo objectForKey : @"NSErrorFailingURLStringKey"]; 
-	NSLog(@"didFailLoadWithError %@", CurrentUrl);
+	NSString* CurrentUrl = [InError.userInfo objectForKey : @"NSErrorFailingURLStringKey"];
+//	NSLog(@"didFailNavigation: %@, error %@", CurrentUrl, InError);
 	WebBrowserWidget->NotifyUrlChanged(CurrentUrl);
 }
 #endif
diff --git a/Engine/Source/Runtime/WebBrowser/Private/IOS/IOSPlatformWebBrowser.h b/Engine/Source/Runtime/WebBrowser/Private/IOS/IOSPlatformWebBrowser.h
--- a/Engine/Source/Runtime/WebBrowser/Private/IOS/IOSPlatformWebBrowser.h
+++ b/Engine/Source/Runtime/WebBrowser/Private/IOS/IOSPlatformWebBrowser.h
@@ -64,8 +64,8 @@ supportsMetal : (bool)InSupportsMetal supportsMetalMRT : (bool)InSupportsMetalMR
 -(void)updateWebViewGLESTexture:(GLuint)gltexture;
 - (void)updateWebViewMetalTexture : (id<MTLTexture>)texture;
 #if !PLATFORM_TVOS
--(BOOL)webView:(UIWebView*)InWebView shouldStartLoadWithRequest:(NSURLRequest*)InRequest navigationType:(UIWebViewNavigationType)InNavigationType;
--(void)webView:(UIWebView*)InWebView didFailLoadWithError:(NSError*)InError;
+- (void)webView:(WKWebView*)InWebView decidePolicyForNavigationAction : (WKNavigationAction*)InNavigationAction decisionHandler : (void(^)(WKNavigationActionPolicy))InDecisionHandler;
+-(void)webView:(WKWebView *)InWebView didCommitNavigation : (WKNavigation *)InNavigation;
 #endif
 @end
 
diff --git a/Engine/Source/Runtime/WebBrowser/Private/WebBrowserSingleton.cpp b/Engine/Source/Runtime/WebBrowser/Private/WebBrowserSingleton.cpp
--- a/Engine/Source/Runtime/WebBrowser/Private/WebBrowserSingleton.cpp
+++ b/Engine/Source/Runtime/WebBrowser/Private/WebBrowserSingleton.cpp
@@ -219,6 +219,8 @@ FWebBrowserSingleton::FWebBrowserSingleton(const FWebBrowserInitSettings& WebBro
 #endif
 	, bDevToolsShortcutEnabled(UE_BUILD_DEBUG)
 	, bJSBindingsToLoweringEnabled(true)
+	, DefaultMaterial(nullptr)
+	, DefaultTranslucentMaterial(nullptr)
 {
 #if WITH_CEF3
 	// The FWebBrowserSingleton must be initialized on the game thread
diff --git a/Engine/Source/Runtime/WebBrowser/Private/WebBrowserSingleton.h b/Engine/Source/Runtime/WebBrowser/Private/WebBrowserSingleton.h
--- a/Engine/Source/Runtime/WebBrowser/Private/WebBrowserSingleton.h
+++ b/Engine/Source/Runtime/WebBrowser/Private/WebBrowserSingleton.h
@@ -39,6 +39,7 @@ class IWebBrowserCookieManager;
 class IWebBrowserWindow;
 struct FWebBrowserWindowInfo;
 struct FWebBrowserInitSettings;
+class UMaterialInterface;
 
 PRAGMA_DISABLE_DEPRECATION_WARNINGS
 
@@ -121,6 +122,30 @@ public:
 		bJSBindingsToLoweringEnabled = bEnabled;
 	}
 
+	/** Set a reference to UWebBrowser's default material*/
+	virtual void SetDefaultMaterial(UMaterialInterface* InDefaultMaterial) override
+	{
+		DefaultMaterial = InDefaultMaterial;
+	}
+
+	/** Set a reference to UWebBrowser's translucent material*/
+	virtual void SetDefaultTranslucentMaterial(UMaterialInterface* InDefaultMaterial) override
+	{
+		DefaultTranslucentMaterial = InDefaultMaterial;
+	}
+
+	/** Get a reference to UWebBrowser's default material*/
+	virtual UMaterialInterface* GetDefaultMaterial() override
+	{
+		return DefaultMaterial;
+	}
+
+	/** Get a reference to UWebBrowser's translucent material*/
+	virtual UMaterialInterface* GetDefaultTranslucentMaterial() override
+	{
+		return DefaultTranslucentMaterial;
+	}
+
 public:
 
 	// FTickerObjectBase Interface
@@ -151,6 +176,11 @@ private:
 
 	bool bJSBindingsToLoweringEnabled;
 
+	/** Reference to UWebBrowser's default material*/
+	UMaterialInterface* DefaultMaterial;
+	/** Reference to UWebBrowser's translucent material*/
+	UMaterialInterface* DefaultTranslucentMaterial;
+
 };
 
 PRAGMA_ENABLE_DEPRECATION_WARNINGS
diff --git a/Engine/Source/Runtime/WebBrowser/Public/IWebBrowserSingleton.h b/Engine/Source/Runtime/WebBrowser/Public/IWebBrowserSingleton.h
--- a/Engine/Source/Runtime/WebBrowser/Public/IWebBrowserSingleton.h
+++ b/Engine/Source/Runtime/WebBrowser/Public/IWebBrowserSingleton.h
@@ -9,6 +9,7 @@ class FCEFWebBrowserWindow;
 class IWebBrowserCookieManager;
 class IWebBrowserWindow;
 class IWebBrowserSchemeHandlerFactory;
+class UMaterialInterface;
 struct FWebBrowserWindowInfo;
 
 class IWebBrowserWindowFactory
@@ -202,4 +203,14 @@ public:
 	 */
 	virtual void SetJSBindingToLoweringEnabled(bool bEnabled) = 0;
 
+
+	/** Set a reference to UWebBrowser's default material*/
+	virtual void SetDefaultMaterial(UMaterialInterface* InDefaultMaterial) = 0;
+	/** Set a reference to UWebBrowser's translucent material*/
+	virtual void SetDefaultTranslucentMaterial(UMaterialInterface* InDefaultMaterial) = 0;
+
+	/** Get a reference to UWebBrowser's default material*/
+	virtual UMaterialInterface* GetDefaultMaterial() = 0;
+	/** Get a reference to UWebBrowser's transparent material*/
+	virtual UMaterialInterface* GetDefaultTranslucentMaterial() = 0;
 };
