UE-59697 WebBrowser SetVisibility not working on Android & iOS

#4.21
#iOS
#Android

SetVisibility implemented on iOS and Android's web browsers, by extending the implementation on Windows: use a list to store references to the webbrowsers in the scene , then in WebBrowserSingleton::Tick check if the cached web browsers windows or widgets ticked on the last frame

diff --git a/Engine/Build/Android/Java/src/com/epicgames/ue4/WebViewControl.java b/Engine/Build/Android/Java/src/com/epicgames/ue4/WebViewControl.java
--- a/Engine/Build/Android/Java/src/com/epicgames/ue4/WebViewControl.java
+++ b/Engine/Build/Android/Java/src/com/epicgames/ue4/WebViewControl.java
@@ -168,6 +168,20 @@ class WebViewControl
 		});
 	}
 
+	boolean PendingSetVisibility;
+	public void SetVisibility(boolean InIsVisible)
+	{
+		PendingSetVisibility = InIsVisible;
+		GameActivity._activity.runOnUiThread(new Runnable()
+		{
+			@Override
+			public void run()
+			{
+				webView.setVisibility(PendingSetVisibility? View.VISIBLE: View.GONE);
+			}
+		});
+	}
+	
 	public boolean didResolutionChange()
 	{
 		if (null != mOESTextureRenderer)
diff --git a/Engine/Source/Runtime/WebBrowser/Private/Android/AndroidJavaWebBrowser.cpp b/Engine/Source/Runtime/WebBrowser/Private/Android/AndroidJavaWebBrowser.cpp
--- a/Engine/Source/Runtime/WebBrowser/Private/Android/AndroidJavaWebBrowser.cpp
+++ b/Engine/Source/Runtime/WebBrowser/Private/Android/AndroidJavaWebBrowser.cpp
@@ -42,6 +42,7 @@ FJavaAndroidWebBrowser::FJavaAndroidWebBrowser(bool swizzlePixels, bool vulkanRe
 	, CloseMethod(GetClassMethod("Close", "()V"))
 	, GoBackOrForwardMethod(GetClassMethod("GoBackOrForward", "(I)V"))
 	, SetAndroid3DBrowserMethod(GetClassMethod("SetAndroid3DBrowser", "(Z)V"))
+	, SetVisibilityMethod(GetClassMethod("SetVisibility", "(Z)V"))
 {
 	VideoTexture = nullptr;
 	bVideoTextureValid = false;
@@ -220,5 +221,10 @@ void FJavaAndroidWebBrowser::SetAndroid3DBrowser(bool InIsAndroid3DBrowser)
 	CallMethod<void>(SetAndroid3DBrowserMethod, InIsAndroid3DBrowser);
 }
 
+void FJavaAndroidWebBrowser::SetVisibility(bool InIsVisible)
+{
+	CallMethod<void>(SetVisibilityMethod, InIsVisible);
+}
+
 #endif // USE_ANDROID_JNI
 
diff --git a/Engine/Source/Runtime/WebBrowser/Private/Android/AndroidJavaWebBrowser.h b/Engine/Source/Runtime/WebBrowser/Private/Android/AndroidJavaWebBrowser.h
--- a/Engine/Source/Runtime/WebBrowser/Private/Android/AndroidJavaWebBrowser.h
+++ b/Engine/Source/Runtime/WebBrowser/Private/Android/AndroidJavaWebBrowser.h
@@ -2,9 +2,10 @@
 
 #pragma once
 
+#include "CoreMinimal.h"
+
 #if USE_ANDROID_JNI
 
-#include "CoreMinimal.h"
 #include "Android/AndroidJava.h"
 #include "RHI.h"
 #include "RHIResources.h"
@@ -28,6 +29,7 @@ public:
 	void GoBack();
 	void GoForward();
 	void SetAndroid3DBrowser(bool InIsAndroid3DBrowser);
+	void SetVisibility(bool InIsVisible);
 	void Update(const int posX, const int posY, const int sizeX, const int sizeY);
 private:
 	static FName GetClassName();
@@ -46,6 +48,7 @@ private:
 	FJavaClassMethod CloseMethod;
 	FJavaClassMethod GoBackOrForwardMethod;
 	FJavaClassMethod SetAndroid3DBrowserMethod;
+	FJavaClassMethod SetVisibilityMethod;
 
 	// FrameUpdateInfo member field ids
 	jclass FrameUpdateInfoClass;
diff --git a/Engine/Source/Runtime/WebBrowser/Private/Android/AndroidWebBrowserWidget.cpp b/Engine/Source/Runtime/WebBrowser/Private/Android/AndroidWebBrowserWidget.cpp
--- a/Engine/Source/Runtime/WebBrowser/Private/Android/AndroidWebBrowserWidget.cpp
+++ b/Engine/Source/Runtime/WebBrowser/Private/Android/AndroidWebBrowserWidget.cpp
@@ -108,9 +108,8 @@ void SAndroidWebBrowserWidget::Construct(const FArguments& Args)
 		AllWebControls.Add(reinterpret_cast<int64>(this), StaticCastSharedRef<SAndroidWebBrowserWidget>(AsShared()));
 	}
 
-	WebBrowserWindowPtr = Args._WebBrowserWindow;
 	IsAndroid3DBrowser = true;
-
+	WebBrowserWindowPtr = Args._WebBrowserWindow;
 	HistorySize = 0;
 	HistoryPosition = 0;
 	
@@ -162,13 +161,17 @@ void SAndroidWebBrowserWidget::Construct(const FArguments& Args)
 
 void SAndroidWebBrowserWidget::Tick(const FGeometry& AllottedGeometry, const double InCurrentTime, const float InDeltaTime)
 {
-	if (WebBrowserWindowPtr.IsValid() && WebBrowserWindowPtr.Pin()->GetParentWindow().IsValid())
+	if(WebBrowserWindowPtr.IsValid())
 	{
-		bool ShouldSetAndroid3DBrowser = WebBrowserWindowPtr.Pin()->GetParentWindow().Get()->IsVirtualWindow();
-		if(IsAndroid3DBrowser != ShouldSetAndroid3DBrowser)
+		WebBrowserWindowPtr.Pin()->SetTickLastFrame();
+		if (WebBrowserWindowPtr.Pin()->GetParentWindow().IsValid())
 		{
-			IsAndroid3DBrowser = ShouldSetAndroid3DBrowser;
-			JavaWebBrowser->SetAndroid3DBrowser(IsAndroid3DBrowser);
+			bool ShouldSetAndroid3DBrowser = WebBrowserWindowPtr.Pin()->GetParentWindow().Get()->IsVirtualWindow();
+			if (IsAndroid3DBrowser != ShouldSetAndroid3DBrowser)
+			{
+				IsAndroid3DBrowser = ShouldSetAndroid3DBrowser;
+				JavaWebBrowser->SetAndroid3DBrowser(IsAndroid3DBrowser);
+			}
 		}
 	}
 
@@ -392,7 +395,9 @@ void SAndroidWebBrowserWidget::Tick(const FGeometry& AllottedGeometry, const dou
 
 int32 SAndroidWebBrowserWidget::OnPaint(const FPaintArgs& Args, const FGeometry& AllottedGeometry, const FSlateRect& MyCullingRect, FSlateWindowElementList& OutDrawElements, int32 LayerId, const FWidgetStyle& InWidgetStyle, bool bParentEnabled) const
 {
-	if (IsAndroid3DBrowser && WebBrowserBrush.IsValid())
+	bool bIsVisible = !WebBrowserWindowPtr.IsValid() || WebBrowserWindowPtr.Pin()->IsVisible();
+
+	if (bIsVisible && IsAndroid3DBrowser && WebBrowserBrush.IsValid())
 	{
 		FSlateDrawElement::MakeBox(OutDrawElements, LayerId, AllottedGeometry.ToPaintGeometry(), WebBrowserBrush.Get(), ESlateDrawEffect::None);
 	}
@@ -456,6 +461,11 @@ bool SAndroidWebBrowserWidget::CanGoForward()
 	return HistoryPosition < HistorySize-1;
 }
 
+void SAndroidWebBrowserWidget::SetWebBrowserVisibility(bool InIsVisible)
+{
+	JavaWebBrowser->SetVisibility(InIsVisible);
+}
+
 jbyteArray SAndroidWebBrowserWidget::HandleShouldInterceptRequest(jstring JUrl)
 {
 	JNIEnv*	JEnv = FAndroidApplication::GetJavaEnv();
diff --git a/Engine/Source/Runtime/WebBrowser/Private/Android/AndroidWebBrowserWidget.h b/Engine/Source/Runtime/WebBrowser/Private/Android/AndroidWebBrowserWidget.h
--- a/Engine/Source/Runtime/WebBrowser/Private/Android/AndroidWebBrowserWidget.h
+++ b/Engine/Source/Runtime/WebBrowser/Private/Android/AndroidWebBrowserWidget.h
@@ -84,6 +84,8 @@ public:
 	// Jobj can either be a WebViewControl, a WebViewControl.ViewClient or WebViewControl.ChromeClient instance
 	static TSharedPtr<SAndroidWebBrowserWidget> GetWidgetPtr(JNIEnv* JEnv, jobject Jobj);
 
+	//set the native control's visibility
+	void SetWebBrowserVisibility(bool InIsVisible);
 protected:
 	static FCriticalSection WebControlsCS;
 	static TMap<int64, TWeakPtr<SAndroidWebBrowserWidget>> AllWebControls;
diff --git a/Engine/Source/Runtime/WebBrowser/Private/Android/AndroidWebBrowserWindow.cpp b/Engine/Source/Runtime/WebBrowser/Private/Android/AndroidWebBrowserWindow.cpp
--- a/Engine/Source/Runtime/WebBrowser/Private/Android/AndroidWebBrowserWindow.cpp
+++ b/Engine/Source/Runtime/WebBrowser/Private/Android/AndroidWebBrowserWindow.cpp
@@ -30,6 +30,8 @@ FAndroidWebBrowserWindow::FAndroidWebBrowserWindow(FString InUrl, TOptional<FStr
 	, ErrorCode(0)
 	, Scripting(new FAndroidJSScripting(bInJSBindingToLoweringEnabled))
 	, AndroidWindowSize(FIntPoint(500, 500))
+	, bTickedLastFrame(true)
+	, bIsVisible(true)
 {
 }
 
@@ -298,4 +300,25 @@ void FAndroidWebBrowserWindow::UnbindUObject(const FString& Name, UObject* Objec
 	Scripting->UnbindUObject(Name, Object, bIsPermanent);
 }
 
+void FAndroidWebBrowserWindow::CheckTickActivity()
+{
+	if (bIsVisible != bTickedLastFrame)
+	{
+		bIsVisible = bTickedLastFrame;
+		BrowserWidget->SetWebBrowserVisibility(bIsVisible);
+	}
+
+	bTickedLastFrame = false;
+}
+
+void FAndroidWebBrowserWindow::SetTickLastFrame()
+{
+	bTickedLastFrame = true;
+}
+
+bool FAndroidWebBrowserWindow::IsVisible()
+{
+	return bIsVisible;
+}
+
 #endif // USE_ANDROID_JNI
\ No newline at end of file
diff --git a/Engine/Source/Runtime/WebBrowser/Private/Android/AndroidWebBrowserWindow.h b/Engine/Source/Runtime/WebBrowser/Private/Android/AndroidWebBrowserWindow.h
--- a/Engine/Source/Runtime/WebBrowser/Private/Android/AndroidWebBrowserWindow.h
+++ b/Engine/Source/Runtime/WebBrowser/Private/Android/AndroidWebBrowserWindow.h
@@ -2,6 +2,8 @@
 
 #pragma once
 
+#include "CoreMinimal.h"
+
 #if USE_ANDROID_JNI
 
 #include "IWebBrowserWindow.h"
@@ -188,7 +190,17 @@ public:
 	/**
 	* Called from the WebBrowserSingleton tick event. Should test whether the widget got a tick from Slate last frame and set the state to hidden if not.
 	*/
-	//void CheckTickActivity();
+	void CheckTickActivity() override;
+
+	/**
+	* Signal from the widget, meaning that the widget is still active
+	*/
+	void SetTickLastFrame();
+
+	/**
+	* Browser's visibility
+	*/
+	bool IsVisible();
 
 private:
 
@@ -276,6 +288,12 @@ private:
 	TSharedPtr<SWindow> ParentWindow;
 
 	FIntPoint AndroidWindowSize;
+
+	/** Used to detect when the widget is hidden*/
+	bool bTickedLastFrame;
+
+	/** Tracks whether the widget is currently visible or not*/
+	bool bIsVisible;
 };
 
 typedef FAndroidWebBrowserWindow FWebBrowserWindow;
diff --git a/Engine/Source/Runtime/WebBrowser/Private/CEF/CEFWebBrowserWindow.h b/Engine/Source/Runtime/WebBrowser/Private/CEF/CEFWebBrowserWindow.h
--- a/Engine/Source/Runtime/WebBrowser/Private/CEF/CEFWebBrowserWindow.h
+++ b/Engine/Source/Runtime/WebBrowser/Private/CEF/CEFWebBrowserWindow.h
@@ -481,7 +481,7 @@ public:
 	/**
 	 * Called from the WebBrowserSingleton tick event. Should test wether the widget got a tick from Slate last frame and set the state to hidden if not.
 	 */
-	void CheckTickActivity();
+	void CheckTickActivity() override;
 
 	/**
 	* Called from the engine tick.
diff --git a/Engine/Source/Runtime/WebBrowser/Private/IOS/IOSPlatformWebBrowser.cpp b/Engine/Source/Runtime/WebBrowser/Private/IOS/IOSPlatformWebBrowser.cpp
--- a/Engine/Source/Runtime/WebBrowser/Private/IOS/IOSPlatformWebBrowser.cpp
+++ b/Engine/Source/Runtime/WebBrowser/Private/IOS/IOSPlatformWebBrowser.cpp
@@ -92,13 +92,17 @@ class SIOSWebBrowserWidget : public SLeafWidget
 	{
 		if (WebViewWrapper != nil)
 		{
-			if (WebBrowserWindowPtr.IsValid() && WebBrowserWindowPtr.Pin()->GetParentWindow().IsValid())
+			if (WebBrowserWindowPtr.IsValid())
 			{
-				bool ShouldSet3DBrowser = WebBrowserWindowPtr.Pin()->GetParentWindow().Get()->IsVirtualWindow();
-				if (IsIOS3DBrowser != ShouldSet3DBrowser)
+				WebBrowserWindowPtr.Pin()->SetTickLastFrame();
+				if (WebBrowserWindowPtr.Pin()->GetParentWindow().IsValid())
 				{
-					IsIOS3DBrowser = ShouldSet3DBrowser;
-					[WebViewWrapper set3D : IsIOS3DBrowser];
+					bool ShouldSet3DBrowser = WebBrowserWindowPtr.Pin()->GetParentWindow().Get()->IsVirtualWindow();
+					if (IsIOS3DBrowser != ShouldSet3DBrowser)
+					{
+						IsIOS3DBrowser = ShouldSet3DBrowser;
+						[WebViewWrapper set3D : IsIOS3DBrowser];
+					}
 				}
 			}
 
@@ -193,7 +197,9 @@ class SIOSWebBrowserWidget : public SLeafWidget
 	int32 OnPaint(const FPaintArgs& Args, const FGeometry& AllottedGeometry, const FSlateRect& MyCullingRect, FSlateWindowElementList& OutDrawElements, int32 LayerId, const FWidgetStyle& InWidgetStyle, bool bParentEnabled) const
 	{
 #if !PLATFORM_TVOS
-		if (IsIOS3DBrowser && WebBrowserBrush.IsValid())
+		bool bIsVisible = !WebBrowserWindowPtr.IsValid() || WebBrowserWindowPtr.Pin()->IsVisible();
+		
+		if (bIsVisible && IsIOS3DBrowser && WebBrowserBrush.IsValid())
 		{
 			FSlateDrawElement::MakeBox(OutDrawElements, LayerId, AllottedGeometry.ToPaintGeometry(), WebBrowserBrush.Get(), ESlateDrawEffect::None);
 		}
@@ -252,6 +258,16 @@ class SIOSWebBrowserWidget : public SLeafWidget
 
 	}
 
+	void SetWebBrowserVisibility(bool InIsVisible)
+	{
+		if (WebViewWrapper != nil)
+		{
+			UE_LOG(LogIOS, Warning, TEXT("SetWebBrowserVisibility %d!"), InIsVisible);
+
+			[WebViewWrapper setVisibility : InIsVisible];
+		}
+	}
+
 	void ExecuteJavascript(const FString& Script)
 	{
 		if (WebViewWrapper != nil)
@@ -348,7 +364,7 @@ supportsMetal : (bool)InSupportsMetal supportsMetalMRT : (bool)InSupportsMetalMR
 			[self.WebView setOpaque : YES];
 		}
 
-		[self setWebViewVisible];
+		[self setDefaultVisibility];
 	});
 #endif
 }
@@ -441,22 +457,43 @@ supportsMetal : (bool)InSupportsMetal supportsMetalMRT : (bool)InSupportsMetalMR
 		{
 			//default is 2D
 			IsIOS3DBrowser = InIsIOS3DBrowser;
-			[self setWebViewVisible];
+			[self setDefaultVisibility];
 		}
 	});
 }
 
--(void)setWebViewVisible;
+-(void)setDefaultVisibility;
 {
 #if !PLATFORM_TVOS
-	if (IsIOS3DBrowser)
+	dispatch_async(dispatch_get_main_queue(), ^
 	{
-		[self.WebViewContainer setHidden : YES];
-	}
-	else
+		if (IsIOS3DBrowser)
+		{
+			[self.WebViewContainer setHidden : YES];
+		}
+		else
+		{
+			[self.WebViewContainer setHidden : NO];
+		}
+	});
+#endif
+}
+
+-(void)setVisibility:(bool)InIsVisible;
+{
+#if !PLATFORM_TVOS
+	dispatch_async(dispatch_get_main_queue(), ^
 	{
-		[self.WebViewContainer setHidden : NO];
-	}
+		if (InIsVisible)
+		{
+			[self setDefaultVisibility];
+		}
+		else
+		{
+			UE_LOG(LogIOS, Warning, TEXT("setVisibility %d HIDE!"), InIsVisible);
+			[self.WebViewContainer setHidden : YES];
+		}
+	});
 #endif
 }
 
@@ -591,6 +628,8 @@ FWebBrowserWindow::FWebBrowserWindow(FString InUrl, TOptional<FString> InContent
 	, ContentsToLoad(MoveTemp(InContentsToLoad))
 	, bUseTransparency(InUseTransparency)
 	, IOSWindowSize(FIntPoint(500, 500))
+	, bTickedLastFrame(true)
+	, bIsVisible(true)
 {
 }
 
@@ -800,4 +839,25 @@ void FWebBrowserWindow::NotifyUrlChanged(const FString& InCurrentUrl)
 		UrlChangedEvent.Broadcast(CurrentUrl);
 	}
 }
+void FWebBrowserWindow::CheckTickActivity()
+{
+	if (bIsVisible != bTickedLastFrame)
+	{
+		bIsVisible = bTickedLastFrame;
+		BrowserWidget->SetWebBrowserVisibility(bIsVisible);
+	}
+
+	bTickedLastFrame = false;
+}
+
+void FWebBrowserWindow::SetTickLastFrame()
+{
+	bTickedLastFrame = true;
+}
+
+bool FWebBrowserWindow::IsVisible()
+{
+	return bIsVisible;
+}
+
 #endif
diff --git a/Engine/Source/Runtime/WebBrowser/Private/IOS/IOSPlatformWebBrowser.h b/Engine/Source/Runtime/WebBrowser/Private/IOS/IOSPlatformWebBrowser.h
--- a/Engine/Source/Runtime/WebBrowser/Private/IOS/IOSPlatformWebBrowser.h
+++ b/Engine/Source/Runtime/WebBrowser/Private/IOS/IOSPlatformWebBrowser.h
@@ -55,7 +55,8 @@ supportsMetal : (bool)InSupportsMetal supportsMetalMRT : (bool)InSupportsMetalMR
 -(void)loadurl:(NSURL*)InURL;
 -(void)executejavascript:(NSString*)InJavaScript;
 -(void)set3D:(bool)InIsIOS3DBrowser;
--(void)setWebViewVisible;
+-(void)setDefaultVisibility;
+-(void)setVisibility:(bool)InIsVisible;
 -(FTextureRHIRef)GetVideoTexture;
 -(void)SetVideoTexture:(FTextureRHIRef)Texture;
 -(void)SetVideoTextureValid:(bool)Condition;
@@ -250,6 +251,20 @@ public:
 	void NotifyUrlChanged(const FString& InCurrentUrl);
 
 public:
+	/**
+	* Called from the WebBrowserSingleton tick event. Should test whether the widget got a tick from Slate last frame and set the state to hidden if not.
+	*/
+	void CheckTickActivity() override;
+
+	/**
+	* Signal from the widget, meaning that the widget is still active
+	*/
+	void SetTickLastFrame();
+
+	/**
+	* Browser's visibility
+	*/
+	bool IsVisible();
 
 private:
 
@@ -317,6 +332,12 @@ private:
 	TSharedPtr<SWindow> ParentWindow;
 
 	FIntPoint IOSWindowSize;
+
+	/** Used to detect when the widget is hidden*/
+	bool bTickedLastFrame;
+
+	/** Tracks whether the widget is currently visible or not*/
+	bool bIsVisible;
 };
 
 #endif
diff --git a/Engine/Source/Runtime/WebBrowser/Private/WebBrowserSingleton.cpp b/Engine/Source/Runtime/WebBrowser/Private/WebBrowserSingleton.cpp
--- a/Engine/Source/Runtime/WebBrowser/Private/WebBrowserSingleton.cpp
+++ b/Engine/Source/Runtime/WebBrowser/Private/WebBrowserSingleton.cpp
@@ -366,6 +366,12 @@ FWebBrowserSingleton::~FWebBrowserSingleton()
 	CEFBrowserApp = nullptr;
 	// Shut down CEF.
 	CefShutdown();
+#elif PLATFORM_IOS || PLATFORM_PS4 || (PLATFORM_ANDROID && USE_ANDROID_JNI)
+	{
+		FScopeLock Lock(&WindowInterfacesCS);
+		// Clear this before CefShutdown() below
+		WindowInterfaces.Reset();
+	}
 #endif
 }
 
@@ -582,6 +588,28 @@ bool FWebBrowserSingleton::Tick(float DeltaTime)
 			}
 		}
 	}
+
+#elif PLATFORM_IOS || PLATFORM_PS4 || (PLATFORM_ANDROID && USE_ANDROID_JNI)
+	FScopeLock Lock(&WindowInterfacesCS);
+	bool bIsSlateAwake = FSlateApplication::IsInitialized() && !FSlateApplication::Get().IsSlateAsleep();
+	// Remove any windows that have been deleted and check whether it's currently visible
+	for (int32 Index = WindowInterfaces.Num() - 1; Index >= 0; --Index)
+	{
+		if (!WindowInterfaces[Index].IsValid())
+		{
+			WindowInterfaces.RemoveAt(Index);
+		}
+		else if (bIsSlateAwake) // only check for Tick activity if Slate is currently ticking
+		{
+			TSharedPtr<IWebBrowserWindow> BrowserWindow = WindowInterfaces[Index].Pin();
+			if (BrowserWindow.IsValid())
+			{
+				// Test if we've ticked recently. If not assume the browser window has become hidden.
+				BrowserWindow->CheckTickActivity();
+			}
+		}
+	}
+
 #endif
 	return true;
 }
diff --git a/Engine/Source/Runtime/WebBrowser/Private/WebBrowserSingleton.h b/Engine/Source/Runtime/WebBrowser/Private/WebBrowserSingleton.h
--- a/Engine/Source/Runtime/WebBrowser/Private/WebBrowserSingleton.h
+++ b/Engine/Source/Runtime/WebBrowser/Private/WebBrowserSingleton.h
@@ -161,15 +161,21 @@ private:
 	void HandleRenderProcessCreated(CefRefPtr<CefListValue> ExtraInfo);
 	/** Pointer to the CEF App implementation */
 	CefRefPtr<FCEFBrowserApp>			CEFBrowserApp;
-	/** List of currently existing browser windows */
-	TArray<TWeakPtr<FCEFWebBrowserWindow>>	WindowInterfaces;
-	/** Critical section for thread safe modification of WindowInterfaces array. */
-	FCriticalSection WindowInterfacesCS;
 
 	TMap<FString, CefRefPtr<CefRequestContext>> RequestContexts;
 	FCefSchemeHandlerFactories SchemeHandlerFactories;
 #endif
 
+	/** List of currently existing browser windows */
+#if WITH_CEF3
+	TArray<TWeakPtr<FCEFWebBrowserWindow>>	WindowInterfaces;
+#elif PLATFORM_IOS || PLATFORM_PS4 || (PLATFORM_ANDROID && USE_ANDROID_JNI)
+	TArray<TWeakPtr<IWebBrowserWindow>>	WindowInterfaces;
+#endif
+
+	/** Critical section for thread safe modification of WindowInterfaces array. */
+	FCriticalSection WindowInterfacesCS;
+
 	TSharedRef<IWebBrowserWindowFactory> WebBrowserWindowFactory;
 
 	bool bDevToolsShortcutEnabled;
diff --git a/Engine/Source/Runtime/WebBrowser/Public/IWebBrowserWindow.h b/Engine/Source/Runtime/WebBrowser/Public/IWebBrowserWindow.h
--- a/Engine/Source/Runtime/WebBrowser/Public/IWebBrowserWindow.h
+++ b/Engine/Source/Runtime/WebBrowser/Public/IWebBrowserWindow.h
@@ -6,6 +6,7 @@
 #include "Input/CursorReply.h"
 #include "Input/Reply.h"
 #include "Widgets/SWindow.h"
+#include "SWebBrowser.h"
 
 class Error;
 class FSlateShaderResource;
@@ -298,6 +299,7 @@ public:
 	*/
 	virtual void SetParentWindow(TSharedPtr<class SWindow> Window) = 0;
 
+	virtual void CheckTickActivity() {};
 public:
 
 	/** A delegate that is invoked when the loading state of a document changed. */
@@ -364,8 +366,25 @@ public:
 	DECLARE_DELEGATE_RetVal_OneParam(bool, FOnDragWindow, const FPointerEvent& /*MouseEvent*/)
 	virtual FOnDragWindow& OnDragWindow() = 0;
 
+	/** A delegate that is invoked to check the visibility of the native browser */
+	DECLARE_DELEGATE_RetVal(bool, FOnCheckVisibility);
+	virtual FOnCheckVisibility& OnCheckVisibility()
+	{
+		return OnCheckVisibilityDelegate;
+	}
+
+	virtual bool CheckVisibility()
+	{
+		return !OnCheckVisibilityDelegate.IsBound() || OnCheckVisibilityDelegate.Execute();
+	}
+
 protected:
 
 	/** Virtual Destructor. */
 	virtual ~IWebBrowserWindow() { };
+
+private:
+	/** Delegate for veritying the window's visibility  */
+	FOnCheckVisibility OnCheckVisibilityDelegate;
+
 };
