Remove unnecessary comment in IOSPlatformWebBrowser.cpp

diff --git a/Engine/Source/Runtime/WebBrowser/Private/IOS/IOSPlatformWebBrowser.cpp b/Engine/Source/Runtime/WebBrowser/Private/IOS/IOSPlatformWebBrowser.cpp
--- a/Engine/Source/Runtime/WebBrowser/Private/IOS/IOSPlatformWebBrowser.cpp
+++ b/Engine/Source/Runtime/WebBrowser/Private/IOS/IOSPlatformWebBrowser.cpp
@@ -490,7 +490,6 @@ supportsMetal : (bool)InSupportsMetal supportsMetalMRT : (bool)InSupportsMetalMR
 		}
 		else
 		{
-			UE_LOG(LogIOS, Warning, TEXT("setVisibility %d HIDE!"), InIsVisible);
 			[self.WebViewContainer setHidden : YES];
 		}
 	});
