UEMOB-185 Implement the equivalent of FAndroidJSScripting for iOS
UE-59488 FWebBrowserWindow::UnbindUObject of IOSPlatformWebBrowser.cpp is not implemented yet

#4.21
#iOS

New methods (using WebKit, matching the Android implementation):
- Go back/forward
- HandlePageLoading (loading/didCommitNavigation, loaded/didFinishNavigation)
- HandleReceivedError (didFailNavigation)
- execute JS code
- controlling the client using JS commands (WKWebViewConfiguration & WKUserContentController):
	- in Android we using shouldInterceptRequest to intercept a custom resource URL (JS: XMLHttpRequest.send, document.location)
	- iOS will register handlers (JS: window.webkit.messageHandlers.<FMobileJSScripting::JSMessageHandler>.postMessage)

diff --git a/Engine/Source/Runtime/WebBrowser/Private/Android/AndroidJSScripting.cpp b/Engine/Source/Runtime/WebBrowser/Private/Android/AndroidJSScripting.cpp
deleted file mode 100644
diff --git a/Engine/Source/Runtime/WebBrowser/Private/Android/AndroidJSScripting.h b/Engine/Source/Runtime/WebBrowser/Private/Android/AndroidJSScripting.h
deleted file mode 100644
diff --git a/Engine/Source/Runtime/WebBrowser/Private/Android/AndroidJSStructDeserializerBackend.cpp b/Engine/Source/Runtime/WebBrowser/Private/Android/AndroidJSStructDeserializerBackend.cpp
deleted file mode 100644
diff --git a/Engine/Source/Runtime/WebBrowser/Private/Android/AndroidJSStructDeserializerBackend.h b/Engine/Source/Runtime/WebBrowser/Private/Android/AndroidJSStructDeserializerBackend.h
deleted file mode 100644
diff --git a/Engine/Source/Runtime/WebBrowser/Private/Android/AndroidJSStructSerializerBackend.cpp b/Engine/Source/Runtime/WebBrowser/Private/Android/AndroidJSStructSerializerBackend.cpp
deleted file mode 100644
diff --git a/Engine/Source/Runtime/WebBrowser/Private/Android/AndroidJSStructSerializerBackend.h b/Engine/Source/Runtime/WebBrowser/Private/Android/AndroidJSStructSerializerBackend.h
deleted file mode 100644
diff --git a/Engine/Source/Runtime/WebBrowser/Private/Android/AndroidWebBrowserDialog.h b/Engine/Source/Runtime/WebBrowser/Private/Android/AndroidWebBrowserDialog.h
--- a/Engine/Source/Runtime/WebBrowser/Private/Android/AndroidWebBrowserDialog.h
+++ b/Engine/Source/Runtime/WebBrowser/Private/Android/AndroidWebBrowserDialog.h
@@ -2,6 +2,8 @@
 
 #pragma once
 
+#include "CoreMinimal.h"
+
 #if USE_ANDROID_JNI
 
 #include "IWebBrowserDialog.h"
diff --git a/Engine/Source/Runtime/WebBrowser/Private/Android/AndroidWebBrowserWidget.cpp b/Engine/Source/Runtime/WebBrowser/Private/Android/AndroidWebBrowserWidget.cpp
--- a/Engine/Source/Runtime/WebBrowser/Private/Android/AndroidWebBrowserWidget.cpp
+++ b/Engine/Source/Runtime/WebBrowser/Private/Android/AndroidWebBrowserWidget.cpp
@@ -6,7 +6,7 @@
 
 #include "AndroidWebBrowserWindow.h"
 #include "AndroidWebBrowserDialog.h"
-#include "AndroidJSScripting.h"
+#include "MobileJS/MobileJSScripting.h"
 #include "Android/AndroidApplication.h"
 #include "Android/AndroidWindow.h"
 #include "Android/AndroidJava.h"
@@ -450,7 +450,6 @@ void SAndroidWebBrowserWidget::GoForward()
 	JavaWebBrowser->GoForward();
 }
 
-
 bool SAndroidWebBrowserWidget::CanGoBack()
 {
 	return HistoryPosition > 1;
@@ -476,7 +475,7 @@ jbyteArray SAndroidWebBrowserWidget::HandleShouldInterceptRequest(jstring JUrl)
 
 	FString Response;
 	bool bOverrideResponse = false;
-	int32 Position = Url.Find(*FAndroidJSScripting::JSMessageTag, ESearchCase::CaseSensitive);
+	int32 Position = Url.Find(*FMobileJSScripting::JSMessageTag, ESearchCase::CaseSensitive);
 	if (Position >= 0)
 	{
 		AsyncTask(ENamedThreads::GameThread, [Url, Position, this]()
@@ -487,7 +486,7 @@ jbyteArray SAndroidWebBrowserWidget::HandleShouldInterceptRequest(jstring JUrl)
 				if (BrowserWindow.IsValid())
 				{
 					FString Origin = Url.Left(Position);
-					FString Message = Url.RightChop(Position + FAndroidJSScripting::JSMessageTag.Len());
+					FString Message = Url.RightChop(Position + FMobileJSScripting::JSMessageTag.Len());
 
 					TArray<FString> Params;
 					Message.ParseIntoArray(Params, TEXT("/"), false);
diff --git a/Engine/Source/Runtime/WebBrowser/Private/Android/AndroidWebBrowserWindow.cpp b/Engine/Source/Runtime/WebBrowser/Private/Android/AndroidWebBrowserWindow.cpp
--- a/Engine/Source/Runtime/WebBrowser/Private/Android/AndroidWebBrowserWindow.cpp
+++ b/Engine/Source/Runtime/WebBrowser/Private/Android/AndroidWebBrowserWindow.cpp
@@ -17,7 +17,7 @@ namespace {
 
 	static const FString JSGetSourceCommand = TEXT("GetSource");
 	static const FString JSMessageGetSourceScript =
-		TEXT("document.location = '") + FAndroidJSScripting::JSMessageTag + JSGetSourceCommand +
+		TEXT("document.location = '") + FMobileJSScripting::JSMessageTag + JSGetSourceCommand +
 		TEXT("/' + encodeURIComponent(document.documentElement.innerHTML);");
 
 }
@@ -28,10 +28,11 @@ FAndroidWebBrowserWindow::FAndroidWebBrowserWindow(FString InUrl, TOptional<FStr
 	, bUseTransparency(InUseTransparency)
 	, DocumentState(EWebBrowserDocumentState::NoDocument)
 	, ErrorCode(0)
-	, Scripting(new FAndroidJSScripting(bInJSBindingToLoweringEnabled))
+	, Scripting(new FMobileJSScripting(bInJSBindingToLoweringEnabled))
 	, AndroidWindowSize(FIntPoint(500, 500))
-	, bTickedLastFrame(true)
+	, bIsDisabled(false)
 	, bIsVisible(true)
+	, bTickedLastFrame(true)
 {
 }
 
@@ -257,6 +258,7 @@ void FAndroidWebBrowserWindow::NotifyDocumentLoadingStateChange(const FString& I
 
 void FAndroidWebBrowserWindow::SetIsDisabled(bool bValue)
 {
+	bIsDisabled = bValue;
 }
 
 TSharedPtr<SWindow> FAndroidWebBrowserWindow::GetParentWindow() const
@@ -313,7 +315,7 @@ void FAndroidWebBrowserWindow::CheckTickActivity()
 
 void FAndroidWebBrowserWindow::SetTickLastFrame()
 {
-	bTickedLastFrame = true;
+	bTickedLastFrame = !bIsDisabled;
 }
 
 bool FAndroidWebBrowserWindow::IsVisible()
diff --git a/Engine/Source/Runtime/WebBrowser/Private/Android/AndroidWebBrowserWindow.h b/Engine/Source/Runtime/WebBrowser/Private/Android/AndroidWebBrowserWindow.h
--- a/Engine/Source/Runtime/WebBrowser/Private/Android/AndroidWebBrowserWindow.h
+++ b/Engine/Source/Runtime/WebBrowser/Private/Android/AndroidWebBrowserWindow.h
@@ -7,8 +7,8 @@
 #if USE_ANDROID_JNI
 
 #include "IWebBrowserWindow.h"
-#include "AndroidJSScripting.h"
 #include "Widgets/SWidget.h"
+#include "MobileJS/MobileJSScripting.h"
 
 class SAndroidWebBrowserWidget;
 class SWebBrowserView;
@@ -281,7 +281,7 @@ private:
 	EWebBrowserDocumentState DocumentState;
 	int ErrorCode;
 
-	FAndroidJSScriptingPtr Scripting;
+	FMobileJSScriptingPtr Scripting;
 
 	mutable TOptional<TFunction<void (const FString&)>> GetPageSourceCallback;
 
@@ -289,11 +289,14 @@ private:
 
 	FIntPoint AndroidWindowSize;
 
-	/** Used to detect when the widget is hidden*/
-	bool bTickedLastFrame;
+	/** Tracks whether the widget is currently disabled or not*/
+	bool bIsDisabled;
 
 	/** Tracks whether the widget is currently visible or not*/
 	bool bIsVisible;
+
+	/** Used to detect when the widget is hidden*/
+	bool bTickedLastFrame;
 };
 
 typedef FAndroidWebBrowserWindow FWebBrowserWindow;
diff --git a/Engine/Source/Runtime/WebBrowser/Private/IOS/IOSPlatformWebBrowser.cpp b/Engine/Source/Runtime/WebBrowser/Private/IOS/IOSPlatformWebBrowser.cpp
--- a/Engine/Source/Runtime/WebBrowser/Private/IOS/IOSPlatformWebBrowser.cpp
+++ b/Engine/Source/Runtime/WebBrowser/Private/IOS/IOSPlatformWebBrowser.cpp
@@ -6,6 +6,8 @@
 #include "IOS/IOSView.h"
 #include "IOS/IOSAppDelegate.h"
 #include "Widgets/SLeafWidget.h"
+#include "MobileJS/MobileJSScripting.h"
+#include "PlatformHttp.h"
 
 #import <UIKit/UIKit.h>
 #import <MetalKit/MetalKit.h>
@@ -227,8 +229,79 @@ class SIOSWebBrowserWidget : public SLeafWidget
 			[WebViewWrapper loadstring : [NSString stringWithUTF8String : TCHAR_TO_UTF8(*InContents)] dummyurl : [NSURL URLWithString : [NSString stringWithUTF8String : TCHAR_TO_UTF8(*InDummyURL)]]];
 		}
 	}
+	
+	void StopLoad()
+	{
+		if (WebViewWrapper != nil)
+		{
+			[WebViewWrapper stopLoading];
+		}
+	}
+
+	void Reload()
+	{
+		if (WebViewWrapper != nil)
+		{
+			[WebViewWrapper reload];
+		}
+	}
+
+	void Close()
+	{
+		if (WebViewWrapper != nil)
+		{
+			[WebViewWrapper close];
+			WebViewWrapper = nil;
+		}
+		WebBrowserWindowPtr.Reset();
+	}
+
+	void GoBack()
+	{
+		if (WebViewWrapper != nil)
+		{
+			[WebViewWrapper goBack];
+		}
+	}
+
+	void GoForward()
+	{
+		if (WebViewWrapper != nil)
+		{
+			[WebViewWrapper goForward];
+		}
+	}
+
+
+	bool CanGoBack()
+	{
+		if (WebViewWrapper != nil)
+		{
+			return [WebViewWrapper canGoBack];
+		}
+		return false;
+	}
+
+	bool CanGoForward()
+	{
+		if (WebViewWrapper != nil)
+		{
+			return [WebViewWrapper canGoForward];
+		}
+		return false;
+	}
+
+	void SetWebBrowserVisibility(bool InIsVisible)
+	{
+		if (WebViewWrapper != nil)
+		{
+			UE_LOG(LogIOS, Warning, TEXT("SetWebBrowserVisibility %d!"), InIsVisible);
 
-	bool OnBeforeBrowse(const FString& Url, const FWebNavigationRequest& RequestDetails)
+			[WebViewWrapper setVisibility : InIsVisible];
+		}
+	}
+
+	bool HandleShouldOverrideUrlLoading(const FString& Url)
 	{
 		bool Retval = false;
 		if (WebBrowserWindowPtr.IsValid())
@@ -238,54 +311,91 @@ class SIOSWebBrowserWidget : public SLeafWidget
 			{
 				if (BrowserWindow->OnBeforeBrowse().IsBound())
 				{
+					FWebNavigationRequest RequestDetails;
+					RequestDetails.bIsRedirect = false;
+					RequestDetails.bIsMainFrame = true; // shouldOverrideUrlLoading is only called on the main frame
+
 					Retval = BrowserWindow->OnBeforeBrowse().Execute(Url, RequestDetails);
+					BrowserWindow->SetTitle("");
 				}
 			}
 		}
 		return Retval;
 	}
-	
-	void NotifyUrlChanged(const FString& InCurrentUrl)
+
+	void HandleReceivedTitle(const FString& Title)
 	{
 		if (WebBrowserWindowPtr.IsValid())
 		{
 			TSharedPtr<FWebBrowserWindow> BrowserWindow = WebBrowserWindowPtr.Pin();
-			if (BrowserWindow.IsValid())
+			if (BrowserWindow.IsValid() && !BrowserWindow->GetTitle().Equals(Title))
 			{
-				BrowserWindow->NotifyUrlChanged(InCurrentUrl);
+				BrowserWindow->SetTitle(Title);
 			}
 		}
-
 	}
 
-	void SetWebBrowserVisibility(bool InIsVisible)
+	void ProcessScriptMessage(const FString& Message)
 	{
-		if (WebViewWrapper != nil)
+		if (WebBrowserWindowPtr.IsValid())
 		{
-			UE_LOG(LogIOS, Warning, TEXT("SetWebBrowserVisibility %d!"), InIsVisible);
+			TSharedPtr<FWebBrowserWindow> BrowserWindow = WebBrowserWindowPtr.Pin();
+			if (BrowserWindow.IsValid())
+			{
+				TArray<FString> Params;
+				Message.ParseIntoArray(Params, TEXT("/"), false);
+				if (Params.Num() > 0)
+				{
+					for (int I = 0; I < Params.Num(); I++)
+					{
+						Params[I] = FPlatformHttp::UrlDecode(Params[I]);
+					}
 
-			[WebViewWrapper setVisibility : InIsVisible];
+					FString Command = Params[0];
+					Params.RemoveAt(0, 1);
+					BrowserWindow->OnJsMessageReceived(Command, Params, "");
+				}
+				else
+				{
+					GLog->Logf(ELogVerbosity::Error, TEXT("Invalid message from browser view: %s"), *Message);
+				}
+			}
 		}
 	}
 
-	void ExecuteJavascript(const FString& Script)
+	void HandlePageLoad(const FString& InCurrentUrl, bool bIsLoading)
 	{
-		if (WebViewWrapper != nil)
+		if (WebBrowserWindowPtr.IsValid())
 		{
-			[WebViewWrapper executejavascript : [NSString stringWithUTF8String : TCHAR_TO_UTF8(*Script)]];
+			TSharedPtr<FWebBrowserWindow> BrowserWindow = WebBrowserWindowPtr.Pin();
+			if (BrowserWindow.IsValid())
+			{
+				BrowserWindow->NotifyDocumentLoadingStateChange(InCurrentUrl, bIsLoading);
+			}
 		}
 	}
 
-	void Close()
+	void HandleReceivedError(int ErrorCode, const FString& InCurrentUrl)
+	{
+		if (WebBrowserWindowPtr.IsValid())
+		{
+			TSharedPtr<FWebBrowserWindow> BrowserWindow = WebBrowserWindowPtr.Pin();
+			if (BrowserWindow.IsValid())
+			{
+				BrowserWindow->NotifyDocumentError(InCurrentUrl, ErrorCode);
+			}
+		}
+	}
+
+	void ExecuteJavascript(const FString& Script)
 	{
 		if (WebViewWrapper != nil)
 		{
-			[WebViewWrapper close];
-			WebViewWrapper = nil;
+			[WebViewWrapper executejavascript : [NSString stringWithUTF8String : TCHAR_TO_UTF8(*Script)]];
 		}
-		WebBrowserWindowPtr.Reset();
 	}
 
+
 	~SIOSWebBrowserWidget()
 	{
 		Close();
@@ -346,8 +456,11 @@ supportsMetal : (bool)InSupportsMetal supportsMetalMRT : (bool)InSupportsMetalMR
 		[self.WebViewContainer setOpaque : NO];
 		[self.WebViewContainer setBackgroundColor : [UIColor clearColor]];
 
+		WKWebViewConfiguration *theConfiguration = [[WKWebViewConfiguration alloc] init];
+		NSString* MessageHandlerName = [NSString stringWithFString : FMobileJSScripting::JSMessageHandler];
+		[theConfiguration.userContentController addScriptMessageHandler:self name: MessageHandlerName];
 
-		WebView = [[WKWebView alloc]initWithFrame:CGRectMake(1, 1, 100, 100)];
+		WebView = [[WKWebView alloc]initWithFrame:CGRectMake(1, 1, 100, 100)  configuration : theConfiguration];
 		[self.WebViewContainer addSubview : WebView];
 		WebView.navigationDelegate = self;
 		WebView.UIDelegate = self;
@@ -422,11 +535,33 @@ supportsMetal : (bool)InSupportsMetal supportsMetalMRT : (bool)InSupportsMetalMR
 #endif
 }
 
+-(NSString *)UrlDecode:(NSString *)stringToDecode
+{
+	NSString *result = [stringToDecode stringByReplacingOccurrencesOfString : @"+" withString:@" "];
+	result = [result stringByReplacingPercentEscapesUsingEncoding : NSUTF8StringEncoding];
+	return result;
+}
+
+-(void)userContentController:(WKUserContentController *)userContentController didReceiveScriptMessage : (WKScriptMessage *)message
+{
+	if ([message.body isKindOfClass : [NSString class]])
+	{
+		NSString *Message = message.body;
+		if (Message != nil)
+		{
+			//NSLog(@"Received message %@", Message);
+			WebBrowserWidget->ProcessScriptMessage(Message);
+		}
+
+	}
+}
+
 -(void)executejavascript:(NSString*)InJavaScript
 {
 #if !PLATFORM_TVOS
 	dispatch_async(dispatch_get_main_queue(), ^
 	{
+	//	NSLog(@"executejavascript %@", InJavaScript);
 		[self.WebView evaluateJavaScript : InJavaScript completionHandler : nil];
 	});
 #endif
@@ -496,6 +631,60 @@ supportsMetal : (bool)InSupportsMetal supportsMetalMRT : (bool)InSupportsMetalMR
 #endif
 }
 
+-(void)stopLoading;
+{
+#if !PLATFORM_TVOS
+	dispatch_async(dispatch_get_main_queue(), ^
+	{
+		[self.WebView stopLoading];
+	});
+#endif
+}
+
+-(void)reload;
+{
+#if !PLATFORM_TVOS
+	dispatch_async(dispatch_get_main_queue(), ^
+	{
+		[self.WebView reload];
+	});
+#endif
+}
+
+-(void)goBack;
+{
+#if !PLATFORM_TVOS
+	dispatch_async(dispatch_get_main_queue(), ^
+	{
+		[self.WebView goBack];
+	});
+#endif
+}
+
+-(void)goForward;
+{
+#if !PLATFORM_TVOS
+	dispatch_async(dispatch_get_main_queue(), ^
+	{
+		[self.WebView goForward];
+	});
+#endif
+}
+
+-(bool)canGoBack;
+{
+#if !PLATFORM_TVOS
+	return [self.WebView canGoBack];
+#endif
+}
+
+-(bool)canGoForward;
+{
+#if !PLATFORM_TVOS
+	return [self.WebView canGoForward];
+#endif
+}
+
 -(FTextureRHIRef)GetVideoTexture;
 {
 	return VideoTexture;
@@ -590,11 +779,7 @@ supportsMetal : (bool)InSupportsMetal supportsMetalMRT : (bool)InSupportsMetalMR
 	// Notify on the game thread
 	[FIOSAsyncTask CreateTaskWithBlock : ^ bool(void)
 	{
-		FWebNavigationRequest RequestDetails;
-		//@todo sz - not sure how to detect redirects here .. InNavigationType == UIWebViewNavigationTypeOther?
-		RequestDetails.bIsRedirect = true;
-		RequestDetails.bIsMainFrame = true;
-		WebBrowserWidget->OnBeforeBrowse(UrlStr, RequestDetails);
+		WebBrowserWidget->HandleShouldOverrideUrlLoading(UrlStr);
 		return true;
 	}];
 
@@ -604,10 +789,21 @@ supportsMetal : (bool)InSupportsMetal supportsMetalMRT : (bool)InSupportsMetalMR
 -(void)webView:(WKWebView *)InWebView didCommitNavigation : (WKNavigation *)InNavigation
 {
 	NSString* CurrentUrl = [self.WebView URL].absoluteString;
+	NSString* Title = [self.WebView title];
+	
 //	NSLog(@"didCommitNavigation: %@", CurrentUrl);
-	WebBrowserWidget->NotifyUrlChanged(CurrentUrl);
+	WebBrowserWidget->HandleReceivedTitle(Title);
+	WebBrowserWidget->HandlePageLoad(CurrentUrl, true);
 }
 
+-(void)webView:(WKWebView *)InWebView didFinishNavigation : (WKNavigation *)InNavigation
+{
+	NSString* CurrentUrl = [self.WebView URL].absoluteString;
+	NSString* Title = [self.WebView title];
+	// NSLog(@"didFinishNavigation: %@", CurrentUrl);
+	WebBrowserWidget->HandleReceivedTitle(Title);
+	WebBrowserWidget->HandlePageLoad(CurrentUrl, false);
+}
 -(void)webView:(WKWebView *)InWebView didFailNavigation : (WKNavigation *)InNavigation withError : (NSError*)InError
 {
 	if (InError.domain == NSURLErrorDomain && InError.code == NSURLErrorCancelled)
@@ -617,18 +813,30 @@ supportsMetal : (bool)InSupportsMetal supportsMetalMRT : (bool)InSupportsMetalMR
 	}
 	NSString* CurrentUrl = [InError.userInfo objectForKey : @"NSErrorFailingURLStringKey"];
 //	NSLog(@"didFailNavigation: %@, error %@", CurrentUrl, InError);
-	WebBrowserWidget->NotifyUrlChanged(CurrentUrl);
+	WebBrowserWidget->HandleReceivedError(InError.code, CurrentUrl);
 }
 #endif
 @end
 
-FWebBrowserWindow::FWebBrowserWindow(FString InUrl, TOptional<FString> InContentsToLoad, bool InShowErrorMessage, bool InThumbMouseButtonNavigation, bool InUseTransparency)
+namespace {
+	static const FString JSGetSourceCommand = TEXT("GetSource");
+	static const FString JSMessageGetSourceScript =
+		TEXT("	window.webkit.messageHandlers.") + FMobileJSScripting::JSMessageHandler + TEXT(".postMessage('")+ JSGetSourceCommand +
+		TEXT("/' + encodeURIComponent(document.documentElement.innerHTML));");
+
+}
+
+FWebBrowserWindow::FWebBrowserWindow(FString InUrl, TOptional<FString> InContentsToLoad, bool InShowErrorMessage, bool InThumbMouseButtonNavigation, bool InUseTransparency, bool bInJSBindingToLoweringEnabled)
 	: CurrentUrl(MoveTemp(InUrl))
 	, ContentsToLoad(MoveTemp(InContentsToLoad))
 	, bUseTransparency(InUseTransparency)
+	, DocumentState(EWebBrowserDocumentState::NoDocument)
+	, ErrorCode(0)
+	, Scripting(new FMobileJSScripting(bInJSBindingToLoweringEnabled))
 	, IOSWindowSize(FIntPoint(500, 500))
-	, bTickedLastFrame(true)
+	, bIsDisabled(false)
 	, bIsVisible(true)
+	, bTickedLastFrame(true)
 {
 }
 
@@ -691,17 +899,17 @@ bool FWebBrowserWindow::IsClosing() const
 
 EWebBrowserDocumentState FWebBrowserWindow::GetDocumentLoadingState() const
 {
-	return EWebBrowserDocumentState::Loading;
+	return DocumentState;
 }
 
 FString FWebBrowserWindow::GetTitle() const
 {
-	return "";
+	return Title;
 }
 
 FString FWebBrowserWindow::GetUrl() const
 {
-	return "";
+	return CurrentUrl;
 }
 
 bool FWebBrowserWindow::OnKeyDown(const FKeyEvent& InKeyEvent)
@@ -768,49 +976,105 @@ void FWebBrowserWindow::OnCaptureLost()
 
 bool FWebBrowserWindow::CanGoBack() const
 {
-	return false;
+	return BrowserWidget->CanGoBack();
 }
 
 void FWebBrowserWindow::GoBack()
 {
+	BrowserWidget->GoBack();
 }
 
 bool FWebBrowserWindow::CanGoForward() const
 {
-	return false;
+	return BrowserWidget->CanGoForward();
 }
 
 void FWebBrowserWindow::GoForward()
 {
+	BrowserWidget->GoForward();
 }
 
 bool FWebBrowserWindow::IsLoading() const
 {
-	return false;
+	return DocumentState != EWebBrowserDocumentState::Loading;
 }
 
 void FWebBrowserWindow::Reload()
 {
+	BrowserWidget->Reload();
 }
 
 void FWebBrowserWindow::StopLoad()
 {
+	BrowserWidget->StopLoad();
 }
 
 void FWebBrowserWindow::GetSource(TFunction<void(const FString&)> Callback) const
 {
-	Callback(FString());
+	//@todo: decide what to do about multiple pending requests
+	GetPageSourceCallback.Emplace(Callback);
+
+	// Ugly hack: Work around the fact that ExecuteJavascript is non-const.
+	const_cast<FWebBrowserWindow*>(this)->ExecuteJavascript(JSMessageGetSourceScript);
 }
 
 int FWebBrowserWindow::GetLoadError()
 {
-	return 0;
+	return ErrorCode;
+}
+
+void FWebBrowserWindow::NotifyDocumentError(const FString& InCurrentUrl, int InErrorCode)
+{
+	if (!CurrentUrl.Equals(InCurrentUrl, ESearchCase::CaseSensitive))
+	{
+		CurrentUrl = InCurrentUrl;
+		UrlChangedEvent.Broadcast(CurrentUrl);
+	}
+
+	ErrorCode = InErrorCode;
+	DocumentState = EWebBrowserDocumentState::Error;
+	DocumentStateChangedEvent.Broadcast(DocumentState);
+}
+
+void FWebBrowserWindow::NotifyDocumentLoadingStateChange(const FString& InCurrentUrl, bool IsLoading)
+{
+	// Ignore a load completed notification if there was an error.
+	// For load started, reset any errors from previous page load.
+	if (IsLoading || DocumentState != EWebBrowserDocumentState::Error)
+	{
+		if (!CurrentUrl.Equals(InCurrentUrl, ESearchCase::CaseSensitive))
+		{
+			CurrentUrl = InCurrentUrl;
+			UrlChangedEvent.Broadcast(CurrentUrl);
+		}
+
+		if (!IsLoading && !InCurrentUrl.StartsWith("javascript:"))
+		{
+			Scripting->PageLoaded(SharedThis(this));
+		}
+		ErrorCode = 0;
+		DocumentState = IsLoading
+			? EWebBrowserDocumentState::Loading
+			: EWebBrowserDocumentState::Completed;
+		DocumentStateChangedEvent.Broadcast(DocumentState);
+	}
+
 }
 
 void FWebBrowserWindow::SetIsDisabled(bool bValue)
 {
+	bIsDisabled = bValue;
 }
 
+TSharedPtr<SWindow> FWebBrowserWindow::GetParentWindow() const
+{
+	return ParentWindow;
+}
+
+void FWebBrowserWindow::SetParentWindow(TSharedPtr<SWindow> Window)
+{
+	ParentWindow = Window;
+}
 
 void FWebBrowserWindow::ExecuteJavascript(const FString& Script)
 {
@@ -822,22 +1086,27 @@ void FWebBrowserWindow::CloseBrowser(bool bForce)
 	BrowserWidget->Close();
 }
 
-void FWebBrowserWindow::BindUObject(const FString& Name, UObject* Object, bool bIsPermanent /*= true*/)
+bool FWebBrowserWindow::OnJsMessageReceived(const FString& Command, const TArray<FString>& Params, const FString& Origin)
 {
+	if (Command.Equals(JSGetSourceCommand, ESearchCase::CaseSensitive) && GetPageSourceCallback.IsSet() && Params.Num() == 1)
+	{
+		GetPageSourceCallback.GetValue()(Params[0]);
+		GetPageSourceCallback.Reset();
+		return true;
+	}
+	return Scripting->OnJsMessageReceived(Command, Params, Origin);
 }
 
-void FWebBrowserWindow::UnbindUObject(const FString& Name, UObject* Object /*= nullptr*/, bool bIsPermanent /*= true*/)
+void FWebBrowserWindow::BindUObject(const FString& Name, UObject* Object, bool bIsPermanent /*= true*/)
 {
+	Scripting->BindUObject(Name, Object, bIsPermanent);
 }
 
-void FWebBrowserWindow::NotifyUrlChanged(const FString& InCurrentUrl)
+void FWebBrowserWindow::UnbindUObject(const FString& Name, UObject* Object /*= nullptr*/, bool bIsPermanent /*= true*/)
 {
-	if (!CurrentUrl.Equals(InCurrentUrl, ESearchCase::CaseSensitive))
-	{
-		CurrentUrl = InCurrentUrl;
-		UrlChangedEvent.Broadcast(CurrentUrl);
-	}
+	Scripting->UnbindUObject(Name, Object, bIsPermanent);
 }
+
 void FWebBrowserWindow::CheckTickActivity()
 {
 	if (bIsVisible != bTickedLastFrame)
@@ -851,7 +1120,7 @@ void FWebBrowserWindow::CheckTickActivity()
 
 void FWebBrowserWindow::SetTickLastFrame()
 {
-	bTickedLastFrame = true;
+	bTickedLastFrame = !bIsDisabled;
 }
 
 bool FWebBrowserWindow::IsVisible()
diff --git a/Engine/Source/Runtime/WebBrowser/Private/IOS/IOSPlatformWebBrowser.h b/Engine/Source/Runtime/WebBrowser/Private/IOS/IOSPlatformWebBrowser.h
--- a/Engine/Source/Runtime/WebBrowser/Private/IOS/IOSPlatformWebBrowser.h
+++ b/Engine/Source/Runtime/WebBrowser/Private/IOS/IOSPlatformWebBrowser.h
@@ -7,6 +7,7 @@
 #if PLATFORM_IOS
 
 #include "IWebBrowserWindow.h"
+#include "MobileJS/MobileJSScripting.h"
 #include "Widgets/SWindow.h"
 #import <UIKit/UIKit.h>
 #if !PLATFORM_TVOS
@@ -26,7 +27,7 @@ class SWebBrowserView;
 * Wrapper to contain the UIWebView and implement its delegate functions
 */
 #if !PLATFORM_TVOS
-@interface IOSWebViewWrapper : NSObject <WKUIDelegate, WKNavigationDelegate>
+@interface IOSWebViewWrapper : NSObject <WKUIDelegate, WKNavigationDelegate, WKScriptMessageHandler>
 #else
 @interface IOSWebViewWrapper : NSObject
 #endif
@@ -57,6 +58,12 @@ supportsMetal : (bool)InSupportsMetal supportsMetalMRT : (bool)InSupportsMetalMR
 -(void)set3D:(bool)InIsIOS3DBrowser;
 -(void)setDefaultVisibility;
 -(void)setVisibility:(bool)InIsVisible;
+-(void)stopLoading;
+-(void)reload;
+-(void)goBack;
+-(void)goForward;
+-(bool)canGoBack;
+-(bool)canGoForward;
 -(FTextureRHIRef)GetVideoTexture;
 -(void)SetVideoTexture:(FTextureRHIRef)Texture;
 -(void)SetVideoTextureValid:(bool)Condition;
@@ -90,7 +97,7 @@ private:
 	* @param InContentsToLoad Optional string to load as a web page.
 	* @param InShowErrorMessage Whether to show an error message in case of loading errors.
 	*/
-	FWebBrowserWindow(FString InUrl, TOptional<FString> InContentsToLoad, bool ShowErrorMessage, bool bThumbMouseButtonNavigation, bool bUseTransparency);
+	FWebBrowserWindow(FString InUrl, TOptional<FString> InContentsToLoad, bool ShowErrorMessage, bool bThumbMouseButtonNavigation, bool bUseTransparency, bool bInJSBindingToLoweringEnabled);
 
 	/**
 	 * Create the SWidget for this WebBrowserWindow
@@ -143,15 +150,8 @@ public:
 	virtual void GetSource(TFunction<void (const FString&)> Callback) const;
 	virtual int GetLoadError() override;
 	virtual void SetIsDisabled(bool bValue) override;
-	virtual TSharedPtr<SWindow> GetParentWindow() const override
-	{
-		return ParentWindow;
-	}
-
-	virtual void SetParentWindow(TSharedPtr<SWindow> Window) override
-	{
-		ParentWindow = Window;
-	}
+	virtual TSharedPtr<SWindow> GetParentWindow() const override;
+	virtual void SetParentWindow(TSharedPtr<SWindow> Window) override;
 
 	// TODO: None of these events are actually called
 
@@ -247,8 +247,12 @@ public:
 	{
 		return DragWindowDelegate;
 	}
+	
+	void NotifyDocumentLoadingStateChange(const FString& InCurrentUrl, bool IsLoading);
+
+	void NotifyDocumentError(const FString& InCurrentUrl, int InErrorCode);
 
-	void NotifyUrlChanged(const FString& InCurrentUrl);
+	bool OnJsMessageReceived(const FString& Command, const TArray<FString>& Params, const FString& Origin);
 
 public:
 	/**
@@ -266,10 +270,19 @@ public:
 	*/
 	bool IsVisible();
 
+	void SetTitle(const FString& InTitle)
+	{
+		Title = InTitle;
+		OnTitleChanged().Broadcast(Title);
+	}
+
 private:
 
 	TSharedPtr<SIOSWebBrowserWidget> BrowserWidget;
 
+	/** Current title of this window. */
+	FString Title;
+
 	/** Current Url of this window. */
 	FString CurrentUrl;
 
@@ -329,15 +342,26 @@ private:
 	/** Delegate that is executed when a drag event is detected in an area of the web page tagged as a drag region. */
 	FOnDragWindow DragWindowDelegate;
 
+	/** Current state of the document being loaded. */
+	EWebBrowserDocumentState DocumentState;
+	int ErrorCode;
+
+	FMobileJSScriptingPtr Scripting;
+
+	mutable TOptional<TFunction<void(const FString&)>> GetPageSourceCallback;
+
 	TSharedPtr<SWindow> ParentWindow;
 
 	FIntPoint IOSWindowSize;
 
-	/** Used to detect when the widget is hidden*/
-	bool bTickedLastFrame;
+	/** Tracks whether the widget is currently disabled or not*/
+	bool bIsDisabled;
 
 	/** Tracks whether the widget is currently visible or not*/
 	bool bIsVisible;
+
+	/** Used to detect when the widget is hidden*/
+	bool bTickedLastFrame;
 };
 
 #endif
diff --git a/Engine/Source/Runtime/WebBrowser/Private/MobileJS/MobileJSScripting.cpp b/Engine/Source/Runtime/WebBrowser/Private/MobileJS/MobileJSScripting.cpp
new file mode 100644
--- /dev/null
+++ b/Engine/Source/Runtime/WebBrowser/Private/MobileJS/MobileJSScripting.cpp
@@ -0,0 +1,560 @@
+// Copyright 1998-2018 Epic Games, Inc. All Rights Reserved.
+
+#include "MobileJSScripting.h"
+
+#if PLATFORM_ANDROID  || PLATFORM_IOS
+
+#include "IWebBrowserWindow.h"
+#include "MobileJSStructSerializerBackend.h"
+#include "MobileJSStructDeserializerBackend.h"
+#include "StructSerializer.h"
+#include "StructDeserializer.h"
+#include "UObject/UnrealType.h"
+#include "Async/Async.h"
+
+// For UrlDecode/Encode
+#include "Http.h"
+
+// Inseterted as a part of an URL to send a message to the front end.
+// Note, we can't use a custom protocol due to cross-domain issues.
+const FString FMobileJSScripting::JSMessageTag = TEXT("/!!com.epicgames.ue4.message/");
+const FString FMobileJSScripting::JSMessageHandler = TEXT("com_epicgames_ue4_message");
+
+
+namespace
+{
+	const FString ExecuteMethodCommand = TEXT("ExecuteUObjectMethod");
+	const FString ScriptingInit =
+		TEXT("(function() {")
+			TEXT("var util = Object.create({")
+
+			// Simple random-based (RFC-4122 version 4) UUID generator.
+			// Version 4 UUIDs have the form xxxxxxxx-xxxx-4xxx-yxxx-xxxxxxxxxxxx where x is any hexadecimal digit and y is one of 8, 9, a, or b
+			// This function returns the UUID as a hex string without the dashes
+			TEXT("uuid: function()")
+			TEXT("{")
+			TEXT("	var b = new Uint8Array(16); window.crypto.getRandomValues(b);")
+			TEXT("	b[6] = b[6]&0xf|0x40; b[8]=b[8]&0x3f|0x80;") // Set the reserved bits to the correct values
+			TEXT("	return Array.prototype.reduce.call(b, function(a,i){return a+((0x100|i).toString(16).substring(1))},'').toUpperCase();")
+			TEXT("}, ")
+
+			// save a callback function in the callback registry
+			// returns the uuid of the callback for passing to the host application
+			// ensures that each function object is only stored once.
+			// (Closures executed multiple times are considered separate objects.)
+			TEXT("registerCallback: function(callback)")
+			TEXT("{")
+			TEXT("	var key;")
+			TEXT("	for(key in this.callbacks)")
+			TEXT("	{")
+			TEXT("		if (!this.callbacks[key].isOneShot && this.callbacks[key].accept === callback)")
+			TEXT("		{")
+			TEXT("			return key;")
+			TEXT("		}")
+			TEXT("	}")
+			TEXT("	key = this.uuid();")
+			TEXT("	this.callbacks[key] = {accept:callback, reject:callback, bIsOneShot:false};")
+			TEXT("	return key;")
+			TEXT("}, ")
+
+			TEXT("registerPromise: function(accept, reject, name)")
+			TEXT("{")
+			TEXT("	var key = this.uuid();")
+			TEXT("	this.callbacks[key] = {accept:accept, reject:reject, bIsOneShot:true, name:name};")
+			TEXT("	return key;")
+			TEXT("}, ")
+
+			// invoke a callback method or promise by uuid
+			TEXT("invokeCallback: function(key, bIsError, args)")
+			TEXT("{")
+			TEXT("	var callback = this.callbacks[key];")
+			TEXT("	if (typeof callback === 'undefined')")
+			TEXT("	{")
+			TEXT("		console.error('Unknown callback id', key);")
+			TEXT("		return;")
+			TEXT("	}")
+			TEXT("	if (callback.bIsOneShot)")
+			TEXT("	{")
+			TEXT("		callback.iwanttodeletethis=true;")
+			TEXT("		delete this.callbacks[key];")
+			TEXT("	}")
+			TEXT("	callback[bIsError?'reject':'accept'].apply(window, args);")
+			TEXT("}, ")
+
+			// convert an argument list to a dictionary of arguments.
+			// The args argument must be an argument object as it uses the callee member to deduce the argument names
+			TEXT("argsToDict: function(args)")
+			TEXT("{")
+			TEXT("	var res = {};")
+			TEXT("	args.callee.toString().match(/\\((.+?)\\)/)[1].split(/\\s*,\\s*/).forEach(function(name, idx){res[name]=args[idx]});")
+			TEXT("	return res;")
+			TEXT("}, ")
+
+			// encodes and sends a message to the host application
+			TEXT("sendMessage: function()")
+			TEXT("{")
+#if PLATFORM_IOS
+			TEXT("	window.webkit.messageHandlers.") + FMobileJSScripting::JSMessageHandler + TEXT(".postMessage(Array.prototype.map.call(arguments,function(e){return encodeURIComponent(e)}).join('/'));")
+#else
+			TEXT("	var req=new XMLHttpRequest();")
+			TEXT("	req.open('GET', '") + FMobileJSScripting::JSMessageTag + TEXT("' + Array.prototype.map.call(arguments,function(e){return encodeURIComponent(e)}).join('/'), true);")
+			TEXT("	req.send(null);")
+#endif
+			TEXT("}, ")
+
+			// uses the above helper methods to execute a method on a uobject instance.
+			// the method set as callee on args needs to be a named function, as the name of the method to invoke is taken from it
+			TEXT("executeMethod: function(id, args)")
+			TEXT("{")
+			TEXT("	var self = this;") // the closures need access to the outer this object
+
+			// In case there are function objects in the argument list, temporarily override Function.toJSON to be able to pass them as callbacks
+			TEXT("	var functionJSON = Function.prototype.toJSON;")
+			TEXT("	Function.prototype.toJSON = function(){ return self.registerCallback(this) };")
+
+			// Create a promise object to return back to the caller and create a callback function to handle the response
+			TEXT("	var promiseID;")
+			TEXT("	var promise = new Promise(function (accept, reject) ")
+			TEXT("	{")
+			TEXT("		promiseID = self.registerPromise(accept, reject, args.callee.name)")
+			TEXT("	});")
+
+			// Actually invoke the method by sending a message to the host app
+			TEXT("	this.sendMessage('") + ExecuteMethodCommand + TEXT("', id, promiseID, args.callee.name, JSON.stringify(this.argsToDict(args)));")
+
+			// Restore Function.toJSON back to its old value (usually undefined) and return the promise object to the caller
+			TEXT("	Function.prototype.toJSON = functionJSON;")
+			TEXT("	return promise;")
+			TEXT("}")
+			TEXT("},{callbacks: {value:{}}});")
+
+			// Create the global window.ue variable
+			TEXT("window.ue = Object.create({}, {'$': {writable: false, configurable:false, enumerable: false, value:util}});")
+		TEXT("})();")
+	;
+	const FString ScriptingPostInit =
+		TEXT("(function() {")
+		TEXT("	document.dispatchEvent(new CustomEvent('ue:ready', {details: window.ue}));")
+		TEXT("})();")
+	;
+
+
+	typedef TSharedRef<TJsonWriter<>> FJsonWriterRef;
+
+	template<typename ValueType> void WriteValue(FJsonWriterRef Writer, const FString& Key, const ValueType& Value)
+	{
+		Writer->WriteValue(Key, Value);
+	}
+	void WriteNull(FJsonWriterRef Writer, const FString& Key)
+	{
+		Writer->WriteNull(Key);
+	}
+	void WriteArrayStart(FJsonWriterRef Writer, const FString& Key)
+	{
+		Writer->WriteArrayStart(Key);
+	}
+	void WriteObjectStart(FJsonWriterRef Writer, const FString& Key)
+	{
+		Writer->WriteObjectStart(Key);
+	}
+	void WriteRaw(FJsonWriterRef Writer, const FString& Key, const FString& Value)
+	{
+		Writer->WriteRawJSONValue(Key, Value);
+	}
+	template<typename ValueType> void WriteValue(FJsonWriterRef Writer, const int, const ValueType& Value)
+	{
+		Writer->WriteValue(Value);
+	}
+	void WriteNull(FJsonWriterRef Writer, int)
+	{
+		Writer->WriteNull();
+	}
+	void WriteArrayStart(FJsonWriterRef Writer, int)
+	{
+		Writer->WriteArrayStart();
+	}
+	void WriteObjectStart(FJsonWriterRef Writer, int)
+	{
+		Writer->WriteObjectStart();
+	}
+	void WriteRaw(FJsonWriterRef Writer, int, const FString& Value)
+	{
+		Writer->WriteRawJSONValue(Value);
+	}
+
+	template<typename KeyType>
+	bool WriteJsParam(FMobileJSScriptingRef Scripting, FJsonWriterRef Writer, const KeyType& Key, FWebJSParam& Param)
+	{
+		switch (Param.Tag)
+		{
+			case FWebJSParam::PTYPE_NULL:
+				WriteNull(Writer, Key);
+				break;
+			case FWebJSParam::PTYPE_BOOL:
+				WriteValue(Writer, Key, Param.BoolValue);
+				break;
+			case FWebJSParam::PTYPE_DOUBLE:
+				WriteValue(Writer, Key, Param.DoubleValue);
+				break;
+			case FWebJSParam::PTYPE_INT:
+				WriteValue(Writer, Key, Param.IntValue);
+				break;
+			case FWebJSParam::PTYPE_STRING:
+				WriteValue(Writer, Key, Param.StringValue);
+				break;
+			case FWebJSParam::PTYPE_OBJECT:
+			{
+				if (Param.ObjectValue == nullptr)
+				{
+					WriteNull(Writer, Key);
+				}
+				else
+				{
+					FString ConvertedObject = Scripting->ConvertObject(Param.ObjectValue);
+					WriteRaw(Writer, Key, ConvertedObject);
+				}
+				break;
+			}
+			case FWebJSParam::PTYPE_STRUCT:
+			{
+				FString ConvertedStruct = Scripting->ConvertStruct(Param.StructValue->GetTypeInfo(), Param.StructValue->GetData());
+				WriteRaw(Writer, Key, ConvertedStruct);
+				break;
+			}
+			case FWebJSParam::PTYPE_ARRAY:
+			{
+				WriteArrayStart(Writer, Key);
+				for(int i=0; i < Param.ArrayValue->Num(); ++i)
+				{
+					WriteJsParam(Scripting, Writer, i, (*Param.ArrayValue)[i]);
+				}
+				Writer->WriteArrayEnd();
+				break;
+			}
+			case FWebJSParam::PTYPE_MAP:
+			{
+				WriteObjectStart(Writer, Key);
+				for(auto& Pair : *Param.MapValue)
+				{
+					WriteJsParam(Scripting, Writer, *Pair.Key, Pair.Value);
+				}
+				Writer->WriteObjectEnd();
+				break;
+			}
+			default:
+				return false;
+		}
+		return true;
+	}
+}
+
+void FMobileJSScripting::BindUObject(const FString& Name, UObject* Object, bool bIsPermanent )
+{
+	const FString ExposedName = GetBindingName(Name, Object);
+	FString Converted = ConvertObject(Object);
+	if (bIsPermanent)
+	{
+		// Each object can only have one permanent binding
+		if (BoundObjects[Object].bIsPermanent)
+		{
+			return;
+		}
+		// Existing permanent objects must be removed first
+		if (PermanentUObjectsByName.Contains(ExposedName))
+		{
+			return;
+		}
+		BoundObjects[Object]={true, -1};
+		PermanentUObjectsByName.Add(ExposedName, Object);
+	}
+
+	TSharedPtr<IWebBrowserWindow> Window = WindowPtr.Pin();
+	if (Window.IsValid())
+	{
+		FString SetValueScript = FString::Printf(TEXT("window.ue['%s'] = %s;"), *ExposedName.ReplaceCharWithEscapedChar(), *Converted);
+		Window->ExecuteJavascript(SetValueScript);
+	}
+}
+
+void FMobileJSScripting::UnbindUObject(const FString& Name, UObject* Object, bool bIsPermanent)
+{
+	const FString ExposedName = GetBindingName(Name, Object);
+	if (bIsPermanent)
+	{
+		// If overriding an existing permanent object, make it non-permanent
+		if (PermanentUObjectsByName.Contains(ExposedName) && (Object == nullptr || PermanentUObjectsByName[ExposedName] == Object))
+		{
+			Object = PermanentUObjectsByName.FindAndRemoveChecked(ExposedName);
+			BoundObjects.Remove(Object);
+			return;
+		}
+		else
+		{
+			return;
+		}
+	}
+
+	TSharedPtr<IWebBrowserWindow> Window = WindowPtr.Pin();
+	if (Window.IsValid())
+	{
+		FString DeleteValueScript = FString::Printf(TEXT("delete window.ue['%s'];"), *ExposedName.ReplaceCharWithEscapedChar());
+		Window->ExecuteJavascript(DeleteValueScript);
+	}
+}
+
+bool FMobileJSScripting::OnJsMessageReceived(const FString& Command, const TArray<FString>& Params, const FString& Origin)
+{
+	bool Result = false;
+	if (Command == ExecuteMethodCommand)
+	{
+		Result = HandleExecuteUObjectMethodMessage(Params);
+	}
+
+	return Result;
+}
+
+FString FMobileJSScripting::ConvertStruct(UStruct* TypeInfo, const void* StructPtr)
+{
+	return TEXT("undefined");
+}
+
+FString FMobileJSScripting::ConvertObject(UObject* Object)
+{
+	RetainBinding(Object);
+	UClass* Class = Object->GetClass();
+
+	bool first = true;
+	FString Result = TEXT("(function(){ return Object.create({");
+	for (TFieldIterator<UFunction> FunctionIt(Class, EFieldIteratorFlags::IncludeSuper); FunctionIt; ++FunctionIt)
+	{
+		UFunction* Function = *FunctionIt;
+		if(!first)
+		{
+			Result.Append(TEXT(","));
+		}
+		else
+		{
+			first = false;
+		}
+		Result.Append(*GetBindingName(Function));
+		Result.Append(TEXT(": function "));
+		Result.Append(*GetBindingName(Function));
+		Result.Append(TEXT(" ("));
+
+		bool firstArg = true;
+		for ( TFieldIterator<UProperty> It(Function); It; ++It )
+		{
+			UProperty* Param = *It;
+			if (Param->PropertyFlags & CPF_Parm && ! (Param->PropertyFlags & CPF_ReturnParm) )
+			{
+				UStructProperty *StructProperty = Cast<UStructProperty>(Param);
+				if (!StructProperty || !StructProperty->Struct->IsChildOf(FWebJSResponse::StaticStruct()))
+				{
+					if(!firstArg)
+					{
+						Result.Append(TEXT(", "));
+					}
+					else
+					{
+						firstArg = false;
+					}
+					Result.Append(*GetBindingName(Param));
+				}
+			}
+		}
+
+		Result.Append(TEXT(")"));
+		Result.Append(TEXT(" {return window.ue.$.executeMethod(this.$id, arguments)}"));
+	}
+	Result.Append(TEXT("},{"));
+	Result.Append(TEXT("$id: {writable: false, configurable:false, enumerable: false, value: '"));
+ 	Result.Append(*PtrToGuid(Object).ToString(EGuidFormats::Digits));
+	Result.Append(TEXT("'}})})()"));
+	return Result;
+}
+
+void FMobileJSScripting::InvokeJSFunction(FGuid FunctionId, int32 ArgCount, FWebJSParam Arguments[], bool bIsError)
+{
+	TSharedPtr<IWebBrowserWindow> Window = WindowPtr.Pin();
+	if (Window.IsValid())
+	{
+		FString CallbackScript = FString::Printf(TEXT("window.ue.$.invokeCallback('%s', %s, "), *FunctionId.ToString(EGuidFormats::Digits), (bIsError)?TEXT("true"):TEXT("false"));
+		{
+			TArray<uint8> Buffer;
+			FMemoryWriter MemoryWriter(Buffer);
+			FJsonWriterRef JsonWriter = TJsonWriter<>::Create(&MemoryWriter);
+			JsonWriter->WriteArrayStart();
+			for(int i=0; i<ArgCount; i++)
+			{
+				WriteJsParam(SharedThis(this), JsonWriter, i, Arguments[i]);
+			}
+			JsonWriter->WriteArrayEnd();
+			CallbackScript.Append((TCHAR*)Buffer.GetData(), Buffer.Num()/sizeof(TCHAR));
+		}
+		CallbackScript.Append(TEXT(")"));
+		Window->ExecuteJavascript(CallbackScript);
+	}
+}
+
+void FMobileJSScripting::InvokeJSFunctionRaw(FGuid FunctionId, const FString& RawJSValue, bool bIsError)
+{
+	TSharedPtr<IWebBrowserWindow> Window = WindowPtr.Pin();
+	if (Window.IsValid())
+	{
+		FString CallbackScript = FString::Printf(TEXT("window.ue.$.invokeCallback('%s', %s, [%s])"),
+			*FunctionId.ToString(EGuidFormats::Digits), (bIsError)?TEXT("true"):TEXT("false"), *RawJSValue);
+		Window->ExecuteJavascript(CallbackScript);
+	}
+}
+
+void FMobileJSScripting::InvokeJSErrorResult(FGuid FunctionId, const FString& Error)
+{
+	FWebJSParam Args[1] = {FWebJSParam(Error)};
+	InvokeJSFunction(FunctionId, 1, Args, true);
+}
+
+bool FMobileJSScripting::HandleExecuteUObjectMethodMessage(const TArray<FString>& MessageArgs)
+{
+	if (MessageArgs.Num() != 4)
+	{
+		return false;
+	}
+
+	FGuid ObjectKey;
+	if (!FGuid::Parse(MessageArgs[0], ObjectKey))
+	{
+		// Invalid GUID
+		return false;
+	}
+	// Get the promise callback and use that to report any results from executing this function.
+	FGuid ResultCallbackId;
+	if (!FGuid::Parse(MessageArgs[1], ResultCallbackId))
+	{
+		// Invalid GUID
+		return false;
+	}
+
+	UObject* Object = GuidToPtr(ObjectKey);
+	if (Object == nullptr)
+	{
+		// Unknown uobject id
+		InvokeJSErrorResult(ResultCallbackId, TEXT("Unknown UObject ID"));
+		return true;
+	}
+
+	FName MethodName = FName(*MessageArgs[2]);
+	UFunction* Function = Object->FindFunction(MethodName);
+	if (!Function)
+	{
+		InvokeJSErrorResult(ResultCallbackId, TEXT("Unknown UObject Function"));
+		return true;
+	}
+
+	// Coerce arguments to function arguments.
+	uint16 ParamsSize = Function->ParmsSize;
+	TArray<uint8> Params;
+	UProperty* ReturnParam = nullptr;
+	UProperty* PromiseParam = nullptr;
+
+	if (ParamsSize > 0)
+	{
+		// Find return parameter and a promise argument if present, as we need to handle them differently
+		for ( TFieldIterator<UProperty> It(Function); It; ++It )
+		{
+			UProperty* Param = *It;
+			if (Param->PropertyFlags & CPF_Parm)
+			{
+				if (Param->PropertyFlags & CPF_ReturnParm)
+				{
+					ReturnParam = Param;
+				}
+				else
+				{
+					UStructProperty *StructProperty = Cast<UStructProperty>(Param);
+					if (StructProperty && StructProperty->Struct->IsChildOf(FWebJSResponse::StaticStruct()))
+					{
+						PromiseParam = Param;
+					}
+				}
+				if (ReturnParam && PromiseParam)
+				{
+					break;
+				}
+			}
+		}
+
+		// UFunction is a subclass of UStruct, so we can treat the arguments as a struct for deserialization
+		Params.AddUninitialized(ParamsSize);
+		Function->InitializeStruct(Params.GetData());
+
+		FMobileJSStructDeserializerBackend Backend = FMobileJSStructDeserializerBackend(SharedThis(this), MessageArgs[3]);
+		FStructDeserializer::Deserialize(Params.GetData(), *Function, Backend);
+	}
+
+	if (PromiseParam)
+	{
+		FWebJSResponse* PromisePtr = PromiseParam->ContainerPtrToValuePtr<FWebJSResponse>(Params.GetData());
+		if (PromisePtr)
+		{
+			*PromisePtr = FWebJSResponse(SharedThis(this), ResultCallbackId);
+		}
+	}
+
+	Object->ProcessEvent(Function, Params.GetData());
+	if ( ! PromiseParam ) // If PromiseParam is set, we assume that the UFunction will ensure it is called with the result
+	{
+		if ( ReturnParam )
+		{
+			FStructSerializerPolicies ReturnPolicies;
+			ReturnPolicies.PropertyFilter = [&](const UProperty* CandidateProperty, const UProperty* ParentProperty)
+			{
+				return ParentProperty != nullptr || CandidateProperty == ReturnParam;
+			};
+			FMobileJSStructSerializerBackend ReturnBackend = FMobileJSStructSerializerBackend(SharedThis(this));
+			FStructSerializer::Serialize(Params.GetData(), *Function, ReturnBackend, ReturnPolicies);
+
+			// Extract the result value from the serialized JSON object:
+			FString ResultJS = ReturnBackend.ToString();
+
+			ResultJS.Append(TEXT("['"));
+			ResultJS.Append(GetBindingName(ReturnParam).ReplaceCharWithEscapedChar());
+			ResultJS.Append(TEXT("']"));
+
+			InvokeJSFunctionRaw(ResultCallbackId, ResultJS, false);
+		}
+		else
+		{
+			InvokeJSFunction(ResultCallbackId, 0, nullptr, false);
+		}
+	}
+	return true;
+}
+
+void FMobileJSScripting::PageLoaded(TSharedRef<class IWebBrowserWindow> InWindow)
+{
+	WindowPtr = InWindow;
+
+	// Expunge temporary objects.
+	for (TMap<UObject*, ObjectBinding>::TIterator It(BoundObjects); It; ++It)
+	{
+		if (!It->Value.bIsPermanent)
+		{
+			It.RemoveCurrent();
+		}
+	}
+
+	FString Script = ScriptingInit;
+	for(auto& Item : PermanentUObjectsByName)
+	{
+		Script.Append(*FString::Printf(TEXT("window.ue['%s'] = %s;"), *Item.Key.ReplaceCharWithEscapedChar(), *ConvertObject(Item.Value)));
+	}
+	Script.Append(ScriptingPostInit);
+	InWindow->ExecuteJavascript(Script);
+}
+
+FMobileJSScripting::FMobileJSScripting(bool bJSBindingToLoweringEnabled)
+	: FWebJSScripting(bJSBindingToLoweringEnabled)
+{
+}
+
+#endif // PLATFORM_ANDROID  || PLATFORM_IOS
\ No newline at end of file
diff --git a/Engine/Source/Runtime/WebBrowser/Private/MobileJS/MobileJSScripting.h b/Engine/Source/Runtime/WebBrowser/Private/MobileJS/MobileJSScripting.h
new file mode 100644
--- /dev/null
+++ b/Engine/Source/Runtime/WebBrowser/Private/MobileJS/MobileJSScripting.h
@@ -0,0 +1,61 @@
+// Copyright 1998-2018 Epic Games, Inc. All Rights Reserved.
+#pragma once
+
+#include "CoreMinimal.h"
+
+#if PLATFORM_ANDROID  || PLATFORM_IOS
+
+#include "WebJSFunction.h"
+#include "WebJSScripting.h"
+
+typedef TSharedRef<class FMobileJSScripting> FMobileJSScriptingRef;
+typedef TSharedPtr<class FMobileJSScripting> FMobileJSScriptingPtr;
+
+/**
+ * Implements handling of bridging UObjects client side with JavaScript renderer side.
+ */
+class FMobileJSScripting
+	: public FWebJSScripting
+	, public TSharedFromThis<FMobileJSScripting>
+{
+public:
+	static const FString JSMessageTag;
+	static const FString JSMessageHandler;
+
+	FMobileJSScripting(bool bJSBindingToLoweringEnabled);
+
+	virtual void BindUObject(const FString& Name, UObject* Object, bool bIsPermanent = true) override;
+	virtual void UnbindUObject(const FString& Name, UObject* Object = nullptr, bool bIsPermanent = true) override;
+
+	/**
+	 * Called when a message was received from the browser process.
+	 *
+	 * @param Command The command sent from the browser.
+	 * @param Params Command-specific data.
+	 * @return true if the message was handled, else false.
+	 */
+	bool OnJsMessageReceived(const FString& Command, const TArray<FString>& Params, const FString& Origin);
+
+	FString ConvertStruct(UStruct* TypeInfo, const void* StructPtr);
+	FString ConvertObject(UObject* Object);
+
+	virtual void InvokeJSFunction(FGuid FunctionId, int32 ArgCount, FWebJSParam Arguments[], bool bIsError=false) override;
+	virtual void InvokeJSErrorResult(FGuid FunctionId, const FString& Error) override;
+	void PageLoaded(TSharedRef<class IWebBrowserWindow> InWindow); // Called on page load
+
+private:
+	void InvokeJSFunctionRaw(FGuid FunctionId, const FString& JSValue, bool bIsError=false);
+	bool IsValid()
+	{
+		return WindowPtr.Pin().IsValid();
+	}
+
+	/** Message handling helpers */
+
+	bool HandleExecuteUObjectMethodMessage(const TArray<FString>& Params);
+
+	/** Pointer to the Mobile Browser for this window. */
+	TWeakPtr<class IWebBrowserWindow> WindowPtr;
+};
+
+#endif // PLATFORM_ANDROID  || PLATFORM_IOS
\ No newline at end of file
diff --git a/Engine/Source/Runtime/WebBrowser/Private/MobileJS/MobileJSStructDeserializerBackend.cpp b/Engine/Source/Runtime/WebBrowser/Private/MobileJS/MobileJSStructDeserializerBackend.cpp
new file mode 100644
--- /dev/null
+++ b/Engine/Source/Runtime/WebBrowser/Private/MobileJS/MobileJSStructDeserializerBackend.cpp
@@ -0,0 +1,106 @@
+// Copyright 1998-2018 Epic Games, Inc. All Rights Reserved.
+
+#include "MobileJSStructDeserializerBackend.h"
+
+#if	PLATFORM_ANDROID || PLATFORM_IOS
+
+#include "MobileJSScripting.h"
+#include "UObject/UnrealType.h"
+#include "Templates/Casts.h"
+
+namespace
+{
+	// @todo: this function is copied from CEFJSStructDeserializerBackend.cpp. Move shared utility code to a common header file
+	/**
+	 * Sets the value of the given property.
+	 *
+	 * @param Property The property to set.
+	 * @param Outer The property that contains the property to be set, if any.
+	 * @param Data A pointer to the memory holding the property's data.
+	 * @param ArrayIndex The index of the element to set (if the property is an array).
+	 * @return true on success, false otherwise.
+	 * @see ClearPropertyValue
+	 */
+	template<typename UPropertyType, typename PropertyType>
+	bool SetPropertyValue( UProperty* Property, UProperty* Outer, void* Data, int32 ArrayIndex, const PropertyType& Value )
+	{
+		PropertyType* ValuePtr = nullptr;
+		UArrayProperty* ArrayProperty = Cast<UArrayProperty>(Outer);
+
+		if (ArrayProperty != nullptr)
+		{
+			if (ArrayProperty->Inner != Property)
+			{
+				return false;
+			}
+
+			FScriptArrayHelper ArrayHelper(ArrayProperty, ArrayProperty->template ContainerPtrToValuePtr<void>(Data));
+			int32 Index = ArrayHelper.AddValue();
+
+			ValuePtr = (PropertyType*)ArrayHelper.GetRawPtr(Index);
+		}
+		else
+		{
+			UPropertyType* TypedProperty = Cast<UPropertyType>(Property);
+
+			if (TypedProperty == nullptr || ArrayIndex >= TypedProperty->ArrayDim)
+			{
+				return false;
+			}
+
+			ValuePtr = TypedProperty->template ContainerPtrToValuePtr<PropertyType>(Data, ArrayIndex);
+		}
+
+		if (ValuePtr == nullptr)
+		{
+			return false;
+		}
+
+		*ValuePtr = Value;
+
+		return true;
+	}
+}
+
+bool FMobileJSStructDeserializerBackend::ReadProperty( UProperty* Property, UProperty* Outer, void* Data, int32 ArrayIndex )
+{
+	switch (GetLastNotation())
+	{
+		case EJsonNotation::String:
+		{
+			if (Property->IsA<UStructProperty>())
+			{
+				UStructProperty* StructProperty = Cast<UStructProperty>(Property);
+
+				if ( StructProperty->Struct == FWebJSFunction::StaticStruct())
+				{
+
+					FGuid CallbackID;
+					if (!FGuid::Parse(GetReader()->GetValueAsString(), CallbackID))
+					{
+						return false;
+					}
+
+					FWebJSFunction CallbackObject(Scripting, CallbackID);
+					return SetPropertyValue<UStructProperty, FWebJSFunction>(Property, Outer, Data, ArrayIndex, CallbackObject);
+				}
+			}
+		}
+		break;
+	}
+
+	// If we reach this, default to parent class behavior
+	return FJsonStructDeserializerBackend::ReadProperty(Property, Outer, Data, ArrayIndex);
+}
+
+FMobileJSStructDeserializerBackend::FMobileJSStructDeserializerBackend(FMobileJSScriptingRef InScripting, const FString& JsonString)
+	: Scripting(InScripting)
+	, JsonData()
+	, Reader(JsonData)
+	, FJsonStructDeserializerBackend(Reader)
+{
+	auto Convert = StringCast<UCS2CHAR>(*JsonString);
+	JsonData.Append((uint8*)Convert.Get(), JsonString.Len() * sizeof(UCS2CHAR));
+}
+
+#endif // PLATFORM_ANDROID || PLATFORM_IOS
\ No newline at end of file
diff --git a/Engine/Source/Runtime/WebBrowser/Private/MobileJS/MobileJSStructDeserializerBackend.h b/Engine/Source/Runtime/WebBrowser/Private/MobileJS/MobileJSStructDeserializerBackend.h
new file mode 100644
--- /dev/null
+++ b/Engine/Source/Runtime/WebBrowser/Private/MobileJS/MobileJSStructDeserializerBackend.h
@@ -0,0 +1,27 @@
+// Copyright 1998-2018 Epic Games, Inc. All Rights Reserved.
+
+#pragma once
+
+#include "CoreMinimal.h"
+
+#if PLATFORM_ANDROID  || PLATFORM_IOS
+
+#include "MobileJSScripting.h"
+#include "Backends/JsonStructDeserializerBackend.h"
+#include "Serialization/MemoryReader.h"
+
+class FMobileJSStructDeserializerBackend
+	: public FJsonStructDeserializerBackend
+{
+public:
+	FMobileJSStructDeserializerBackend(FMobileJSScriptingRef InScripting, const FString& JsonString);
+
+	virtual bool ReadProperty( UProperty* Property, UProperty* Outer, void* Data, int32 ArrayIndex ) override;
+
+private:
+	FMobileJSScriptingRef Scripting;
+	TArray<uint8> JsonData;
+	FMemoryReader Reader;
+};
+
+#endif // USE_ANDROID_JNI || PLATFORM_IOS
\ No newline at end of file
diff --git a/Engine/Source/Runtime/WebBrowser/Private/MobileJS/MobileJSStructSerializerBackend.cpp b/Engine/Source/Runtime/WebBrowser/Private/MobileJS/MobileJSStructSerializerBackend.cpp
new file mode 100644
--- /dev/null
+++ b/Engine/Source/Runtime/WebBrowser/Private/MobileJS/MobileJSStructSerializerBackend.cpp
@@ -0,0 +1,62 @@
+// Copyright 1998-2018 Epic Games, Inc. All Rights Reserved.
+
+#include "MobileJSStructSerializerBackend.h"
+
+#if PLATFORM_ANDROID || PLATFORM_IOS
+
+#include "MobileJSScripting.h"
+#include "UObject/UnrealType.h"
+#include "UObject/PropertyPortFlags.h"
+#include "Templates/Casts.h"
+
+void FMobileJSStructSerializerBackend::WriteProperty(const FStructSerializerState& State, int32 ArrayIndex)
+{
+	// The parent class serialzes UObjects as NULLs
+	if (State.ValueType == UObjectProperty::StaticClass())
+	{
+		WriteUObject(State, CastChecked<UObjectProperty>(State.ValueProperty)->GetPropertyValue_InContainer(State.ValueData, ArrayIndex));
+	}
+	// basic property type (json serializable)
+	else
+	{
+		FJsonStructSerializerBackend::WriteProperty(State, ArrayIndex);
+	}
+}
+
+void FMobileJSStructSerializerBackend::WriteUObject(const FStructSerializerState& State, UObject* Value)
+{
+	// Note this function uses WriteRawJSONValue to append non-json data to the output stream.
+	FString RawValue = Scripting->ConvertObject(Value);
+	if ((State.ValueProperty == nullptr) || (State.ValueProperty->ArrayDim > 1) || (State.ValueProperty->GetOuter()->GetClass() == UArrayProperty::StaticClass()))
+	{
+		GetWriter()->WriteRawJSONValue(RawValue);
+	}
+	else if (State.KeyProperty != nullptr)
+	{
+		FString KeyString;
+		State.KeyProperty->ExportTextItem(KeyString, State.KeyData, nullptr, nullptr, PPF_None);
+		GetWriter()->WriteRawJSONValue(KeyString, RawValue);
+	}
+	else
+	{
+		GetWriter()->WriteRawJSONValue(Scripting->GetBindingName(State.ValueProperty), RawValue);
+	}
+}
+
+FString FMobileJSStructSerializerBackend::ToString()
+{
+	ReturnBuffer.Add(0);
+	ReturnBuffer.Add(0);
+	auto CastObject = StringCast<TCHAR>((UCS2CHAR*)ReturnBuffer.GetData());
+	return FString(CastObject.Get(), CastObject.Length());
+}
+
+FMobileJSStructSerializerBackend::FMobileJSStructSerializerBackend(TSharedRef<class FMobileJSScripting> InScripting)
+	: Scripting(InScripting)
+	, ReturnBuffer()
+	, Writer(ReturnBuffer)
+	, FJsonStructSerializerBackend(Writer)
+{
+}
+
+#endif // PLATFORM_ANDROID || PLATFORM_IOS
\ No newline at end of file
diff --git a/Engine/Source/Runtime/WebBrowser/Private/MobileJS/MobileJSStructSerializerBackend.h b/Engine/Source/Runtime/WebBrowser/Private/MobileJS/MobileJSStructSerializerBackend.h
new file mode 100644
--- /dev/null
+++ b/Engine/Source/Runtime/WebBrowser/Private/MobileJS/MobileJSStructSerializerBackend.h
@@ -0,0 +1,45 @@
+// Copyright 1998-2018 Epic Games, Inc. All Rights Reserved.
+
+#pragma once
+
+#include "CoreMinimal.h"
+
+#if PLATFORM_ANDROID || PLATFORM_IOS
+
+#include "MobileJSScripting.h"
+#include "Backends/JsonStructSerializerBackend.h"
+
+class UObject;
+
+/**
+ * Implements a writer for UStruct serialization using JavaScript.
+ *
+ * Based on FJsonStructSerializerBackend, it adds support for certain object types not representable in pure JSON
+ *
+ */
+class FMobileJSStructSerializerBackend
+	: public FJsonStructSerializerBackend
+{
+public:
+
+	/**
+	 * Creates and initializes a new instance.
+	 *
+	 * @param InScripting An instance of a web browser scripting obnject.
+	 */
+	FMobileJSStructSerializerBackend(FMobileJSScriptingRef InScripting);
+
+public:
+	virtual void WriteProperty(const FStructSerializerState& State, int32 ArrayIndex = 0) override;
+
+	FString ToString();
+
+private:
+	void WriteUObject(const FStructSerializerState& State, UObject* Value);
+
+	FMobileJSScriptingRef Scripting;
+	TArray<uint8> ReturnBuffer;
+	FMemoryWriter Writer;
+};
+
+#endif // PLATFORM_ANDROID || PLATFORM_IOS
\ No newline at end of file
diff --git a/Engine/Source/Runtime/WebBrowser/Private/WebBrowserSingleton.cpp b/Engine/Source/Runtime/WebBrowser/Private/WebBrowserSingleton.cpp
--- a/Engine/Source/Runtime/WebBrowser/Private/WebBrowserSingleton.cpp
+++ b/Engine/Source/Runtime/WebBrowser/Private/WebBrowserSingleton.cpp
@@ -532,16 +532,39 @@ TSharedPtr<IWebBrowserWindow> FWebBrowserSingleton::CreateBrowserWindow(const FC
 		WindowSettings.bUseTransparency,
 		bJSBindingsToLoweringEnabled));
 
+	{
+		FScopeLock Lock(&WindowInterfacesCS);
+		WindowInterfaces.Add(NewBrowserWindow);
+	}
 	return NewBrowserWindow;
-#elif PLATFORM_IOS || PLATFORM_PS4
+#elif PLATFORM_IOS
 	// Create new window
 	TSharedPtr<FWebBrowserWindow> NewBrowserWindow = MakeShareable(new FWebBrowserWindow(
 		WindowSettings.InitialURL, 
 		WindowSettings.ContentsToLoad, 
 		WindowSettings.bShowErrorMessage, 
 		WindowSettings.bThumbMouseButtonNavigation, 
+		WindowSettings.bUseTransparency,
+		bJSBindingsToLoweringEnabled));
+
+	{
+		FScopeLock Lock(&WindowInterfacesCS);
+		WindowInterfaces.Add(NewBrowserWindow);
+	}
+	return NewBrowserWindow;
+#elif PLATFORM_PS4
+	// Create new window
+	TSharedPtr<FWebBrowserWindow> NewBrowserWindow = MakeShareable(new FWebBrowserWindow(
+		WindowSettings.InitialURL,
+		WindowSettings.ContentsToLoad,
+		WindowSettings.bShowErrorMessage,
+		WindowSettings.bThumbMouseButtonNavigation,
 		WindowSettings.bUseTransparency));
 
+	{
+		FScopeLock Lock(&WindowInterfacesCS);
+		WindowInterfaces.Add(NewBrowserWindow);
+	}
 	return NewBrowserWindow;
 #endif
 	return nullptr;
