Fix CIS issues in IOSPlatformWebBrowser

diff --git a/Engine/Source/Runtime/WebBrowser/Private/IOS/IOSPlatformWebBrowser.cpp b/Engine/Source/Runtime/WebBrowser/Private/IOS/IOSPlatformWebBrowser.cpp
--- a/Engine/Source/Runtime/WebBrowser/Private/IOS/IOSPlatformWebBrowser.cpp
+++ b/Engine/Source/Runtime/WebBrowser/Private/IOS/IOSPlatformWebBrowser.cpp
@@ -538,7 +538,7 @@ supportsMetal : (bool)InSupportsMetal supportsMetalMRT : (bool)InSupportsMetalMR
 -(NSString *)UrlDecode:(NSString *)stringToDecode
 {
 	NSString *result = [stringToDecode stringByReplacingOccurrencesOfString : @"+" withString:@" "];
-	result = [result stringByReplacingPercentEscapesUsingEncoding : NSUTF8StringEncoding];
+	result = [result stringByRemovingPercentEncoding];
 	return result;
 }
 
diff --git a/Engine/Source/Runtime/WebBrowser/Private/MobileJS/MobileJSStructDeserializerBackend.cpp b/Engine/Source/Runtime/WebBrowser/Private/MobileJS/MobileJSStructDeserializerBackend.cpp
--- a/Engine/Source/Runtime/WebBrowser/Private/MobileJS/MobileJSStructDeserializerBackend.cpp
+++ b/Engine/Source/Runtime/WebBrowser/Private/MobileJS/MobileJSStructDeserializerBackend.cpp
@@ -94,10 +94,10 @@ bool FMobileJSStructDeserializerBackend::ReadProperty( UProperty* Property, UPro
 }
 
 FMobileJSStructDeserializerBackend::FMobileJSStructDeserializerBackend(FMobileJSScriptingRef InScripting, const FString& JsonString)
-	: Scripting(InScripting)
+	: FJsonStructDeserializerBackend(Reader)
+	, Scripting(InScripting)
 	, JsonData()
 	, Reader(JsonData)
-	, FJsonStructDeserializerBackend(Reader)
 {
 	auto Convert = StringCast<UCS2CHAR>(*JsonString);
 	JsonData.Append((uint8*)Convert.Get(), JsonString.Len() * sizeof(UCS2CHAR));
diff --git a/Engine/Source/Runtime/WebBrowser/Private/MobileJS/MobileJSStructSerializerBackend.cpp b/Engine/Source/Runtime/WebBrowser/Private/MobileJS/MobileJSStructSerializerBackend.cpp
--- a/Engine/Source/Runtime/WebBrowser/Private/MobileJS/MobileJSStructSerializerBackend.cpp
+++ b/Engine/Source/Runtime/WebBrowser/Private/MobileJS/MobileJSStructSerializerBackend.cpp
@@ -52,10 +52,10 @@ FString FMobileJSStructSerializerBackend::ToString()
 }
 
 FMobileJSStructSerializerBackend::FMobileJSStructSerializerBackend(TSharedRef<class FMobileJSScripting> InScripting)
-	: Scripting(InScripting)
+	: FJsonStructSerializerBackend(Writer)
+	, Scripting(InScripting)
 	, ReturnBuffer()
 	, Writer(ReturnBuffer)
-	, FJsonStructSerializerBackend(Writer)
 {
 }
 
