UE-63953 TVOS compile errors in Main

#TvOS
#4.21

Fix compile errors by
- commenting out the method (void)userContentController:(WKUserContentController *) for PLATFORM_TVOS
- adding default return values for some methods

diff --git a/Engine/Source/Runtime/WebBrowser/Private/IOS/IOSPlatformWebBrowser.cpp b/Engine/Source/Runtime/WebBrowser/Private/IOS/IOSPlatformWebBrowser.cpp
--- a/Engine/Source/Runtime/WebBrowser/Private/IOS/IOSPlatformWebBrowser.cpp
+++ b/Engine/Source/Runtime/WebBrowser/Private/IOS/IOSPlatformWebBrowser.cpp
@@ -542,6 +542,7 @@ supportsMetal : (bool)InSupportsMetal supportsMetalMRT : (bool)InSupportsMetalMR
 	return result;
 }
 
+#if !PLATFORM_TVOS
 -(void)userContentController:(WKUserContentController *)userContentController didReceiveScriptMessage : (WKScriptMessage *)message
 {
 	if ([message.body isKindOfClass : [NSString class]])
@@ -555,6 +556,7 @@ supportsMetal : (bool)InSupportsMetal supportsMetalMRT : (bool)InSupportsMetalMR
 
 	}
 }
+#endif
 
 -(void)executejavascript:(NSString*)InJavaScript
 {
@@ -673,14 +675,18 @@ supportsMetal : (bool)InSupportsMetal supportsMetalMRT : (bool)InSupportsMetalMR
 
 -(bool)canGoBack;
 {
-#if !PLATFORM_TVOS
+#if PLATFORM_TVOS
+	return false;
+#else
 	return [self.WebView canGoBack];
 #endif
 }
 
 -(bool)canGoForward;
 {
-#if !PLATFORM_TVOS
+#if PLATFORM_TVOS
+	return false;
+#else
 	return [self.WebView canGoForward];
 #endif
 }
