Fix Lumin compile error.

#rnx

diff --git a/Engine/Source/Runtime/WebBrowser/WebBrowser.Build.cs b/Engine/Source/Runtime/WebBrowser/WebBrowser.Build.cs
--- a/Engine/Source/Runtime/WebBrowser/WebBrowser.Build.cs
+++ b/Engine/Source/Runtime/WebBrowser/WebBrowser.Build.cs
@@ -25,6 +25,7 @@ public class WebBrowser : ModuleRules
 		);
 
 		if (Target.Platform == UnrealTargetPlatform.Android ||
+			Target.Platform == UnrealTargetPlatform.Lumin ||
 		    Target.Platform == UnrealTargetPlatform.IOS ||
 		    Target.Platform == UnrealTargetPlatform.TVOS)
 		{
