Disable the WebBrowser module on Lumin to fix CIS.

#rnx

diff --git a/Engine/Source/Runtime/WebBrowser/WebBrowser.Build.cs b/Engine/Source/Runtime/WebBrowser/WebBrowser.Build.cs
--- a/Engine/Source/Runtime/WebBrowser/WebBrowser.Build.cs
+++ b/Engine/Source/Runtime/WebBrowser/WebBrowser.Build.cs
@@ -25,7 +25,6 @@ public class WebBrowser : ModuleRules
 		);
 
 		if (Target.Platform == UnrealTargetPlatform.Android ||
-			Target.Platform == UnrealTargetPlatform.Lumin ||
 		    Target.Platform == UnrealTargetPlatform.IOS ||
 		    Target.Platform == UnrealTargetPlatform.TVOS)
 		{
@@ -80,6 +79,11 @@ public class WebBrowser : ModuleRules
 			PrivateDependencyModuleNames.Add("Engine");
 		}
 
+		if(Target.Platform == UnrealTargetPlatform.Lumin)
+		{
+			PrecompileForTargets = ModuleRules.PrecompileTargetsType.None;
+		}
+
 		bEnableShadowVariableWarnings = false;
 	}
 }
