Fix IOS/Android shadow variable warnings.

diff --git a/Engine/Source/Runtime/WebBrowser/Private/Android/AndroidWebBrowserWidget.cpp b/Engine/Source/Runtime/WebBrowser/Private/Android/AndroidWebBrowserWidget.cpp
--- a/Engine/Source/Runtime/WebBrowser/Private/Android/AndroidWebBrowserWidget.cpp
+++ b/Engine/Source/Runtime/WebBrowser/Private/Android/AndroidWebBrowserWidget.cpp
@@ -310,9 +310,9 @@ void SAndroidWebBrowserWidget::Tick(const FGeometry& AllottedGeometry, const dou
 					if (VideoTexture == nullptr)
 					{
 						FRHIResourceCreateInfo CreateInfo;
-						FRHICommandListImmediate& RHICmdList = FRHICommandListExecutor::GetImmediateCommandList();
-						FIntPoint Size = Params.Size;
-						VideoTexture = RHICmdList.CreateTextureExternal2D(Size.X, Size.Y, PF_R8G8B8A8, 1, 1, 0, CreateInfo);
+						FRHICommandListImmediate& LocalCmdList = FRHICommandListExecutor::GetImmediateCommandList();
+						FIntPoint LocalSize = Params.Size;
+						VideoTexture = LocalCmdList.CreateTextureExternal2D(LocalSize.X, LocalSize.Y, PF_R8G8B8A8, 1, 1, 0, CreateInfo);
 						PinnedJavaWebBrowser->SetVideoTexture(VideoTexture);
 
 						if (VideoTexture == nullptr)
@@ -322,7 +322,7 @@ void SAndroidWebBrowserWidget::Tick(const FGeometry& AllottedGeometry, const dou
 						}
 
 						PinnedJavaWebBrowser->SetVideoTextureValid(false);
-						FPlatformMisc::LowLevelOutputDebugStringf(TEXT("Fetch RT: Created VideoTexture: %d - %s (%d, %d)"), *reinterpret_cast<int32*>(VideoTexture->GetNativeResource()), *Params.PlayerGuid.ToString(), Size.X, Size.Y);
+						FPlatformMisc::LowLevelOutputDebugStringf(TEXT("Fetch RT: Created VideoTexture: %d - %s (%d, %d)"), *reinterpret_cast<int32*>(VideoTexture->GetNativeResource()), *Params.PlayerGuid.ToString(), LocalSize.X, LocalSize.Y);
 					}
 
 					int32 TextureId = *reinterpret_cast<int32*>(VideoTexture->GetNativeResource());
diff --git a/Engine/Source/Runtime/WebBrowser/Private/IOS/IOSPlatformWebBrowser.cpp b/Engine/Source/Runtime/WebBrowser/Private/IOS/IOSPlatformWebBrowser.cpp
--- a/Engine/Source/Runtime/WebBrowser/Private/IOS/IOSPlatformWebBrowser.cpp
+++ b/Engine/Source/Runtime/WebBrowser/Private/IOS/IOSPlatformWebBrowser.cpp
@@ -158,9 +158,9 @@ class SIOSWebBrowserWidget : public SLeafWidget
 							if (VideoTexture == nullptr)
 							{
 								FRHIResourceCreateInfo CreateInfo;
-								FRHICommandListImmediate& RHICmdList = FRHICommandListExecutor::GetImmediateCommandList();
+								FRHICommandListImmediate& LocalRHICmdList = FRHICommandListExecutor::GetImmediateCommandList();
 								FIntPoint Size = Params.Size;
-								VideoTexture = RHICmdList.CreateTextureExternal2D(Size.X, Size.Y, PF_R8G8B8A8, 1, 1, 0, CreateInfo);
+								VideoTexture = LocalRHICmdList.CreateTextureExternal2D(Size.X, Size.Y, PF_R8G8B8A8, 1, 1, 0, CreateInfo);
 								[NativeWebBrowser SetVideoTexture : VideoTexture];
 								//UE_LOG(LogIOS, Log, TEXT("NativeWebBrowser SetVideoTexture:VideoTexture!"));
 
