Append the chrome user agent string to our custom browser product version string.  This will make our embedded browser more familiar to third party pages.  Currently hard coding the string representing our version ("Chrome/59.0.3071.15") and appending it to our custom user agent string.

diff --git a/Engine/Source/Runtime/WebBrowser/Private/WebBrowserModule.cpp b/Engine/Source/Runtime/WebBrowser/Private/WebBrowserModule.cpp
--- a/Engine/Source/Runtime/WebBrowser/Private/WebBrowserModule.cpp
+++ b/Engine/Source/Runtime/WebBrowser/Private/WebBrowserModule.cpp
@@ -14,7 +14,7 @@ DEFINE_LOG_CATEGORY(LogWebBrowser);
 static FWebBrowserSingleton* WebBrowserSingleton = nullptr;
 
 FWebBrowserInitSettings::FWebBrowserInitSettings()
-	: ProductVersion(FString::Printf(TEXT("%s/%s UnrealEngine/%s"), FApp::GetProjectName(), FApp::GetBuildVersion(), *FEngineVersion::Current().ToString()))
+	: ProductVersion(FString::Printf(TEXT("%s/%s UnrealEngine/%s Chrome/59.0.3071.15"), FApp::GetProjectName(), FApp::GetBuildVersion(), *FEngineVersion::Current().ToString()))
 {
 }
 
