UE-64749 App crashes on device when a Widget Blueprint in use contains a Web Browser widget

#Android
#iOS
#4.21

Moving SharedThis(this) out of the WB Window's constructor

diff --git a/Engine/Source/Runtime/WebBrowser/Private/Android/AndroidWebBrowserWindow.cpp b/Engine/Source/Runtime/WebBrowser/Private/Android/AndroidWebBrowserWindow.cpp
--- a/Engine/Source/Runtime/WebBrowser/Private/Android/AndroidWebBrowserWindow.cpp
+++ b/Engine/Source/Runtime/WebBrowser/Private/Android/AndroidWebBrowserWindow.cpp
@@ -28,7 +28,7 @@ FAndroidWebBrowserWindow::FAndroidWebBrowserWindow(FString InUrl, TOptional<FStr
 	, bUseTransparency(InUseTransparency)
 	, DocumentState(EWebBrowserDocumentState::NoDocument)
 	, ErrorCode(0)
-	, Scripting(new FMobileJSScripting(bInJSBindingToLoweringEnabled, SharedThis(this)))
+	, Scripting(new FMobileJSScripting(bInJSBindingToLoweringEnabled))
 	, AndroidWindowSize(FIntPoint(500, 500))
 	, bIsDisabled(false)
 	, bIsVisible(true)
@@ -60,6 +60,9 @@ TSharedRef<SWidget> FAndroidWebBrowserWindow::CreateWidget()
 		.WebBrowserWindow(SharedThis(this));
 
 	BrowserWidget = BrowserWidgetRef;
+
+	Scripting->SetWindow(SharedThis(this));
+		
 	return BrowserWidgetRef;
 }
 
diff --git a/Engine/Source/Runtime/WebBrowser/Private/IOS/IOSPlatformWebBrowser.cpp b/Engine/Source/Runtime/WebBrowser/Private/IOS/IOSPlatformWebBrowser.cpp
--- a/Engine/Source/Runtime/WebBrowser/Private/IOS/IOSPlatformWebBrowser.cpp
+++ b/Engine/Source/Runtime/WebBrowser/Private/IOS/IOSPlatformWebBrowser.cpp
@@ -838,7 +838,7 @@ FWebBrowserWindow::FWebBrowserWindow(FString InUrl, TOptional<FString> InContent
 	, bUseTransparency(InUseTransparency)
 	, DocumentState(EWebBrowserDocumentState::NoDocument)
 	, ErrorCode(0)
-	, Scripting(new FMobileJSScripting(bInJSBindingToLoweringEnabled, SharedThis(this)))
+	, Scripting(new FMobileJSScripting(bInJSBindingToLoweringEnabled))
 	, IOSWindowSize(FIntPoint(500, 500))
 	, bIsDisabled(false)
 	, bIsVisible(true)
@@ -870,6 +870,9 @@ TSharedRef<SWidget> FWebBrowserWindow::CreateWidget()
 		.WebBrowserWindow(SharedThis(this));
 
 	BrowserWidget = BrowserWidgetRef;
+
+	Scripting->SetWindow(SharedThis(this));
+
 	return BrowserWidgetRef;
 }
 
diff --git a/Engine/Source/Runtime/WebBrowser/Private/MobileJS/MobileJSScripting.cpp b/Engine/Source/Runtime/WebBrowser/Private/MobileJS/MobileJSScripting.cpp
--- a/Engine/Source/Runtime/WebBrowser/Private/MobileJS/MobileJSScripting.cpp
+++ b/Engine/Source/Runtime/WebBrowser/Private/MobileJS/MobileJSScripting.cpp
@@ -623,8 +623,12 @@ void FMobileJSScripting::PageLoaded(TSharedRef<class IWebBrowserWindow> InWindow
 	InWindow->ExecuteJavascript(Script);
 }
 
-FMobileJSScripting::FMobileJSScripting(bool bJSBindingToLoweringEnabled, TSharedRef<class IWebBrowserWindow> InWindow)
+FMobileJSScripting::FMobileJSScripting(bool bJSBindingToLoweringEnabled)
 	: FWebJSScripting(bJSBindingToLoweringEnabled)
+{
+}
+
+void FMobileJSScripting::SetWindow(TSharedRef<class IWebBrowserWindow> InWindow)
 {
 	WindowPtr = InWindow;
 }
diff --git a/Engine/Source/Runtime/WebBrowser/Private/MobileJS/MobileJSScripting.h b/Engine/Source/Runtime/WebBrowser/Private/MobileJS/MobileJSScripting.h
--- a/Engine/Source/Runtime/WebBrowser/Private/MobileJS/MobileJSScripting.h
+++ b/Engine/Source/Runtime/WebBrowser/Private/MobileJS/MobileJSScripting.h
@@ -22,7 +22,7 @@ public:
 	static const FString JSMessageTag;
 	static const FString JSMessageHandler;
 
-	FMobileJSScripting(bool bJSBindingToLoweringEnabled, TSharedRef<class IWebBrowserWindow> InWindow);
+	FMobileJSScripting(bool bJSBindingToLoweringEnabled);
 
 	virtual void BindUObject(const FString& Name, UObject* Object, bool bIsPermanent = true) override;
 	virtual void UnbindUObject(const FString& Name, UObject* Object = nullptr, bool bIsPermanent = true) override;
@@ -45,6 +45,8 @@ public:
 	virtual void InvokeJSErrorResult(FGuid FunctionId, const FString& Error) override;
 	void PageLoaded(TSharedRef<class IWebBrowserWindow> InWindow); // Called on page load
 
+	void SetWindow(TSharedRef<class IWebBrowserWindow> InWindow);
+
 private:
 	void InitializeScript(TSharedRef<class IWebBrowserWindow> InWindow);
 	void InvokeJSFunctionRaw(FGuid FunctionId, const FString& JSValue, bool bIsError=false);
