Missing fix from CL4418318 Add external URL calls to purchase flow on Android

#4.21
#Android
#iOS

diff --git a/Engine/Source/Runtime/WebBrowser/Private/MobileJS/MobileJSScripting.cpp b/Engine/Source/Runtime/WebBrowser/Private/MobileJS/MobileJSScripting.cpp
--- a/Engine/Source/Runtime/WebBrowser/Private/MobileJS/MobileJSScripting.cpp
+++ b/Engine/Source/Runtime/WebBrowser/Private/MobileJS/MobileJSScripting.cpp
@@ -552,11 +552,11 @@ bool FMobileJSScripting::HandleExecuteUObjectMethodMessage(const TArray<FString>
 
 			// Extract the result value from the serialized JSON object:
 			FString ResultJS = ReturnBackend.ToString();
-
+			/*
 			ResultJS.Append(TEXT("['"));
 			ResultJS.Append(GetBindingName(ReturnParam).ReplaceCharWithEscapedChar());
 			ResultJS.Append(TEXT("']"));
-
+			*/
 			InvokeJSFunctionRaw(ResultCallbackId, ResultJS, false);
 		}
 		else
