Fix compile issue with Android cookie manager

#android
#lumin

diff --git a/Engine/Source/Runtime/WebBrowser/Private/Android/AndroidCookieManager.cpp b/Engine/Source/Runtime/WebBrowser/Private/Android/AndroidCookieManager.cpp
--- a/Engine/Source/Runtime/WebBrowser/Private/Android/AndroidCookieManager.cpp
+++ b/Engine/Source/Runtime/WebBrowser/Private/Android/AndroidCookieManager.cpp
@@ -2,15 +2,14 @@
 
 #include "AndroidCookieManager.h"
 
-#if PLATFORM_ANDROID
+#if USE_ANDROID_JNI
 
 #include "Android/AndroidJNI.h"
 #include "Android/AndroidApplication.h"
 
-#if USE_ANDROID_JNI
 #include <jni.h>
 #include "Async/TaskGraphInterfaces.h"
-#endif
+
 
 FAndroidCookieManager::FAndroidCookieManager()
 {
@@ -74,4 +73,36 @@ void FAndroidCookieManager::DeleteCookies(const FString& URL, const FString& Coo
 	}
 }
 
+#else
+
+FAndroidCookieManager::FAndroidCookieManager()
+{
+}
+
+FAndroidCookieManager::~FAndroidCookieManager()
+{
+}
+
+void FAndroidCookieManager::SetCookie(const FString& URL, const FCookie& Cookie, TFunction<void(bool)> Completed)
+{
+	bool bResult = false;
+
+	// send result
+	if (Completed)
+	{
+		Completed(bResult);
+	}
+}
+
+void FAndroidCookieManager::DeleteCookies(const FString& URL, const FString& CookieName, TFunction<void(int)> Completed)
+{
+	bool bResult = false;
+
+	// send result
+	if (Completed)
+	{
+		Completed(bResult);
+	}
+}
+
 #endif
