UE-59012 //UE4/Main - Compile UE4Game Lumin - WARNING: Precompiled module 'PurchaseFlow' will not be usable without non-precompiled module 'WebBrowser'.
-Potential fix for this, just allow the module to be precompiled.  These plugins are not enabled by default, so it won't typically be used.

diff --git a/Engine/Source/Runtime/WebBrowser/WebBrowser.Build.cs b/Engine/Source/Runtime/WebBrowser/WebBrowser.Build.cs
--- a/Engine/Source/Runtime/WebBrowser/WebBrowser.Build.cs
+++ b/Engine/Source/Runtime/WebBrowser/WebBrowser.Build.cs
@@ -38,10 +38,6 @@ public class WebBrowser : ModuleRules
 			// We need this one on Android for URL decoding
 			PrivateDependencyModuleNames.Add("HTTP");
 		}
-		if (Target.Platform == UnrealTargetPlatform.Lumin)
-		{
-			PrecompileForTargets = PrecompileTargetsType.None;
-		}
 
 		if (Target.Platform == UnrealTargetPlatform.Win64
 		||  Target.Platform == UnrealTargetPlatform.Win32
