Rollback //UE4/Dev-VR/Engine/Source/Runtime/WebBrowser/WebBrowser.Build.cs to revision 18

diff --git a/Engine/Source/Runtime/WebBrowser/WebBrowser.Build.cs b/Engine/Source/Runtime/WebBrowser/WebBrowser.Build.cs
--- a/Engine/Source/Runtime/WebBrowser/WebBrowser.Build.cs
+++ b/Engine/Source/Runtime/WebBrowser/WebBrowser.Build.cs
@@ -38,6 +38,10 @@ public class WebBrowser : ModuleRules
 			// We need this one on Android for URL decoding
 			PrivateDependencyModuleNames.Add("HTTP");
 		}
+		if (Target.Platform == UnrealTargetPlatform.Lumin)
+		{
+			PrecompileForTargets = PrecompileTargetsType.None;
+		}
 
 		if (Target.Platform == UnrealTargetPlatform.Win64
 		||  Target.Platform == UnrealTargetPlatform.Win32
