Fix shadowed variable warning.

#rnx

diff --git a/Engine/Source/Runtime/WebBrowser/Private/Native/NativeJSScripting.cpp b/Engine/Source/Runtime/WebBrowser/Private/Native/NativeJSScripting.cpp
--- a/Engine/Source/Runtime/WebBrowser/Private/Native/NativeJSScripting.cpp
+++ b/Engine/Source/Runtime/WebBrowser/Private/Native/NativeJSScripting.cpp
@@ -484,7 +484,7 @@ bool FNativeJSScripting::HandleExecuteUObjectMethodMessage(const TArray<FString>
 
 FString FNativeJSScripting::GetInitializeScript()
 {
-	const FString ScriptingInit =
+	const FString NativeScriptingInit =
 		TEXT("(function() {")
 			TEXT("var util = Object.create({")
 
@@ -611,7 +611,7 @@ FString FNativeJSScripting::GetInitializeScript()
 		TEXT("})();")
 	;
 
-	return ScriptingInit;
+	return NativeScriptingInit;
 }
 
 void FNativeJSScripting::PageLoaded()
