Fix compatibility issue with old versions of android browser.  Replace 'let' with 'var'.

diff --git a/Engine/Source/Runtime/WebBrowser/Private/Native/NativeJSScripting.cpp b/Engine/Source/Runtime/WebBrowser/Private/Native/NativeJSScripting.cpp
--- a/Engine/Source/Runtime/WebBrowser/Private/Native/NativeJSScripting.cpp
+++ b/Engine/Source/Runtime/WebBrowser/Private/Native/NativeJSScripting.cpp
@@ -561,7 +561,7 @@ FString FNativeJSScripting::GetInitializeScript()
 			TEXT("{")
 			// @todo: Each kairos native browser will have a different way of passing a message out, here we use webkit postmessage but we'll need
 			//    to be aware of our target platform when generating this script and adjust accordingly
-			TEXT("  let delimiter = '/';")
+			TEXT("  var delimiter = '/';")
 
 #if PLATFORM_ANDROID
 			TEXT("  if(window.JSBridge){")
