[NO NON-PII COMMIT MESSAGE]

diff --git a/Engine/Source/Runtime/WebBrowser/Private/Native/NativeWebBrowserProxy.h b/Engine/Source/Runtime/WebBrowser/Private/Native/NativeWebBrowserProxy.h
--- a/Engine/Source/Runtime/WebBrowser/Private/Native/NativeWebBrowserProxy.h
+++ b/Engine/Source/Runtime/WebBrowser/Private/Native/NativeWebBrowserProxy.h
@@ -132,6 +132,21 @@ public:
 		return BeforePopupDelegate;
 	}
 
+	virtual FOnBeforeResourceLoadDelegate& OnBeforeResourceLoad() override
+	{
+		return BeforeResourceLoadDelegate;
+	}
+
+	virtual FOnResourceLoadCompleteDelegate& OnResourceLoadComplete() override
+	{
+		return ResourceLoadCompleteDelegate;
+	}
+
+	virtual FOnConsoleMessageDelegate& OnConsoleMessage() override
+	{
+		return ConsoleMessageDelegate;
+	}
+
 	DECLARE_DERIVED_EVENT(FNativeWebBrowserProxy, IWebBrowserWindow::FOnShowPopup, FOnShowPopup);
 	virtual FOnShowPopup& OnShowPopup() override
 	{
@@ -199,6 +214,15 @@ private:
 
 	/** Delegate for notifying that a popup window is attempting to open. */
 	FOnBeforePopupDelegate BeforePopupDelegate;
+	
+	/** Delegate that is invoked before the browser loads a resource. Its primary purpose is to inject headers into the request. */
+	FOnBeforeResourceLoadDelegate BeforeResourceLoadDelegate;
+
+	/** Delegate that is invoked on completion of browser resource loads. Its primary purpose is to allow response to failures. */
+	FOnResourceLoadCompleteDelegate ResourceLoadCompleteDelegate;
+
+	/** Delegate that is invoked for each console message */
+	FOnConsoleMessageDelegate ConsoleMessageDelegate;
 
 	/** Delegate for handling requests to create new windows. */
 	FOnCreateWindow CreateWindowDelegate;
