Fixes Ocean compile failure.
WebBrowser has a path where it writes directly to underlying texture objects. Removed OpenGL code.

diff --git a/Engine/Source/Runtime/WebBrowser/Private/IOS/IOSPlatformWebBrowser.cpp b/Engine/Source/Runtime/WebBrowser/Private/IOS/IOSPlatformWebBrowser.cpp
--- a/Engine/Source/Runtime/WebBrowser/Private/IOS/IOSPlatformWebBrowser.cpp
+++ b/Engine/Source/Runtime/WebBrowser/Private/IOS/IOSPlatformWebBrowser.cpp
@@ -41,6 +41,8 @@ class SIOSWebBrowserWidget : public SLeafWidget
 
 		bool bSupportsMetal = false;
 		GConfig->GetBool(TEXT("/Script/IOSRuntimeSettings.IOSRuntimeSettings"), TEXT("bSupportsMetal"), bSupportsMetal, GEngineIni);
+		// At this point we MUST be a Metal renderer.
+		check(bSupportsMetal);
 
 		WebViewWrapper = [IOSWebViewWrapper alloc];
 		[WebViewWrapper create : TSharedPtr<SIOSWebBrowserWidget>(this) useTransparency : Args._UseTransparency supportsMetal : bSupportsMetal supportsMetalMRT : bSupportsMetalMRT];
@@ -716,43 +718,17 @@ supportsMetal : (bool)InSupportsMetal supportsMetalMRT : (bool)InSupportsMetalMR
 {
 #if !PLATFORM_TVOS
 	@synchronized(self) // Briefly block render thread
-		{
-		if (bSupportsMetal)
-		{
-			id<MTLTexture> ptrToMetalTexture = (id<MTLTexture>)ptr;
-			NSUInteger width = [ptrToMetalTexture width];
-			NSUInteger height = [ptrToMetalTexture height];
+	{
+		id<MTLTexture> ptrToMetalTexture = (id<MTLTexture>)ptr;
+		NSUInteger width = [ptrToMetalTexture width];
+		NSUInteger height = [ptrToMetalTexture height];
 
-			[self updateWebViewMetalTexture : ptrToMetalTexture];
-		}
-		else
-		{
-			GLuint glTexture = (GLuint)*reinterpret_cast<int32*>(ptr);
-			glBindTexture(GL_TEXTURE_2D, glTexture);
-			[self updateWebViewGLESTexture : glTexture];
-		}
-		}
+		[self updateWebViewMetalTexture : ptrToMetalTexture];
+	}
 #endif
 	return true;
 }
 
--(void)updateWebViewGLESTexture:(GLuint)gltexture
-{
-#if !PLATFORM_TVOS
-	// create a suitable CoreGraphics context
-	CGColorSpaceRef colourSpace = CGColorSpaceCreateDeviceRGB();
-	CGContextRef context =
-		CGBitmapContextCreate(&gltexture, WebView.bounds.size.width, WebView.bounds.size.height, 8, 4 * WebView.bounds.size.width, colourSpace, kCGImageAlphaPremultipliedLast | kCGBitmapByteOrder32Big);
-	CGColorSpaceRelease(colourSpace);
-	// draw the view to the buffer
-	[WebView.layer renderInContext : context];
-	// upload to OpenGL
-	glTexImage2D(GL_TEXTURE_2D, 0, GL_RGBA, WebView.bounds.size.width, WebView.bounds.size.height, 0, GL_RGBA, GL_UNSIGNED_BYTE, &gltexture);
-	// clean up
-	CGContextRelease(context);
-#endif
-}
-
 -(void)updateWebViewMetalTexture:(id<MTLTexture>)texture
 {
 #if !PLATFORM_TVOS
diff --git a/Engine/Source/Runtime/WebBrowser/Private/IOS/IOSPlatformWebBrowser.h b/Engine/Source/Runtime/WebBrowser/Private/IOS/IOSPlatformWebBrowser.h
--- a/Engine/Source/Runtime/WebBrowser/Private/IOS/IOSPlatformWebBrowser.h
+++ b/Engine/Source/Runtime/WebBrowser/Private/IOS/IOSPlatformWebBrowser.h
@@ -69,7 +69,6 @@ supportsMetal : (bool)InSupportsMetal supportsMetalMRT : (bool)InSupportsMetalMR
 -(void)SetVideoTextureValid:(bool)Condition;
 -(bool)IsVideoTextureValid;
 -(bool)UpdateVideoFrame:(void*)ptr;
--(void)updateWebViewGLESTexture:(GLuint)gltexture;
 - (void)updateWebViewMetalTexture : (id<MTLTexture>)texture;
 #if !PLATFORM_TVOS
 - (void)webView:(WKWebView*)InWebView decidePolicyForNavigationAction : (WKNavigationAction*)InNavigationAction decisionHandler : (void(^)(WKNavigationActionPolicy))InDecisionHandler;
