Fixed issue where web browser UObject binding could cause a memory stomp when invoking UFunctions called by the webpage

diff --git a/Engine/Source/Runtime/WebBrowser/Private/CEF/CEFJSScripting.cpp b/Engine/Source/Runtime/WebBrowser/Private/CEF/CEFJSScripting.cpp
--- a/Engine/Source/Runtime/WebBrowser/Private/CEF/CEFJSScripting.cpp
+++ b/Engine/Source/Runtime/WebBrowser/Private/CEF/CEFJSScripting.cpp
@@ -282,7 +282,7 @@ bool FCEFJSScripting::HandleExecuteUObjectMethodMessage(CefRefPtr<CefListValue>
 		}
 
 		// UFunction is a subclass of UStruct, so we can treat the arguments as a struct for deserialization
-		Params.AddUninitialized(ParamsSize);
+		Params.AddUninitialized(Function->GetStructureSize());
 		Function->InitializeStruct(Params.GetData());
 		FCEFJSStructDeserializerBackend Backend = FCEFJSStructDeserializerBackend(SharedThis(this), NamedArgs);
 		FStructDeserializer::Deserialize(Params.GetData(), *Function, Backend);
diff --git a/Engine/Source/Runtime/WebBrowser/Private/MobileJS/MobileJSScripting.cpp b/Engine/Source/Runtime/WebBrowser/Private/MobileJS/MobileJSScripting.cpp
--- a/Engine/Source/Runtime/WebBrowser/Private/MobileJS/MobileJSScripting.cpp
+++ b/Engine/Source/Runtime/WebBrowser/Private/MobileJS/MobileJSScripting.cpp
@@ -521,7 +521,7 @@ bool FMobileJSScripting::HandleExecuteUObjectMethodMessage(const TArray<FString>
 		}
 
 		// UFunction is a subclass of UStruct, so we can treat the arguments as a struct for deserialization
-		Params.AddUninitialized(ParamsSize);
+		Params.AddUninitialized(Function->GetStructureSize());
 		Function->InitializeStruct(Params.GetData());
 
 		FMobileJSStructDeserializerBackend Backend = FMobileJSStructDeserializerBackend(SharedThis(this), MessageArgs[3]);
