Fixed some issues with browser getting stuck.
CEF Browser's OnScheuleMessagePumpWork method will now properly handle cases where a negative delay was provided.
Prevent scope lock of pump countdown when running CefDoMessageLoopWork().

Note: Copy of CL# 10876841.

diff --git a/Engine/Source/Runtime/WebBrowser/Private/CEF/CEFBrowserApp.cpp b/Engine/Source/Runtime/WebBrowser/Private/CEF/CEFBrowserApp.cpp
--- a/Engine/Source/Runtime/WebBrowser/Private/CEF/CEFBrowserApp.cpp
+++ b/Engine/Source/Runtime/WebBrowser/Private/CEF/CEFBrowserApp.cpp
@@ -30,15 +30,13 @@ void FCEFBrowserApp::OnScheduleMessagePumpWork(int64 delay_ms)
 {
 	FScopeLock Lock(&MessagePumpCountdownCS);
 
-	if (MessagePumpCountdown == -1)
+	// As per CEF documentation, if delay_ms is <= 0, then the call to CefDoMessageLoopWork should happen reasonably soon.  If delay_ms is > 0, then the call
+	//  to CefDoMessageLoopWork should be scheduled to happen after the specified delay and any currently pending scheduled call should be canceled.
+	if(delay_ms < 0)
 	{
-		MessagePumpCountdown = delay_ms;
-	}
-	else
-	{
-		// override if we need the pump to happen sooner
-		MessagePumpCountdown = FMath::Clamp<int64>(MessagePumpCountdown, 0, delay_ms);
+		delay_ms = 0;
 	}
+	MessagePumpCountdown = delay_ms;
 }
 #endif
 
@@ -49,22 +47,29 @@ void FCEFBrowserApp::TickMessagePump(float DeltaTime, bool bForce)
 	return;
 #endif
 
-	FScopeLock Lock(&MessagePumpCountdownCS);
-
-	bool bPump = false;	
-	// count down in order to call message pump
-	if (MessagePumpCountdown >= 0)
+	bool bPump = false;
 	{
-		MessagePumpCountdown -= DeltaTime * 1000;
-		if (MessagePumpCountdown <= 0)
+		FScopeLock Lock(&MessagePumpCountdownCS);
+		
+		// count down in order to call message pump
+		if (MessagePumpCountdown >= 0)
 		{
-			bPump = true;
+			MessagePumpCountdown -= DeltaTime * 1000;
+			if (MessagePumpCountdown <= 0)
+			{
+				bPump = true;
+			}
+			
+			if (bPump || bForce)
+			{
+				// -1 indicates that no countdown is currently happening
+				MessagePumpCountdown = -1;
+			}
 		}
 	}
+	
 	if (bPump || bForce)
 	{
-		// -1 indicates that no countdown is currently happening
-		MessagePumpCountdown = -1;
 		CefDoMessageLoopWork();
 	}
 }
