Fix build for macOS.

diff --git a/Engine/Source/Runtime/WebBrowser/Private/WebBrowserSingleton.cpp b/Engine/Source/Runtime/WebBrowser/Private/WebBrowserSingleton.cpp
--- a/Engine/Source/Runtime/WebBrowser/Private/WebBrowserSingleton.cpp
+++ b/Engine/Source/Runtime/WebBrowser/Private/WebBrowserSingleton.cpp
@@ -310,7 +310,7 @@ FWebBrowserSingleton::FWebBrowserSingleton(const FWebBrowserInitSettings& WebBro
 		LocalePakPath = ResourcesPath + TEXT("/") + LocaleCode + TEXT(".lproj/locale.pak");
 		if (FPaths::FileExists(LocalePakPath))
 		{
-			CefString(&Settings.locale) = *LocaleCode;
+			CefString(&Settings.locale) = TCHAR_TO_WCHAR(*LocaleCode);
 		}
 	}
 #endif
