Portal and Web CIS fixes

diff --git a/Engine/Source/Runtime/WebBrowser/Private/MobileJS/MobileJSStructSerializerBackend.cpp b/Engine/Source/Runtime/WebBrowser/Private/MobileJS/MobileJSStructSerializerBackend.cpp
--- a/Engine/Source/Runtime/WebBrowser/Private/MobileJS/MobileJSStructSerializerBackend.cpp
+++ b/Engine/Source/Runtime/WebBrowser/Private/MobileJS/MobileJSStructSerializerBackend.cpp
@@ -12,7 +12,7 @@
 void FMobileJSStructSerializerBackend::WriteProperty(const FStructSerializerState& State, int32 ArrayIndex)
 {
 	// The parent class serialzes UObjects as NULLs
-	if (State.ValueType == FObjectProperty::StaticClass())
+	if (State.FieldType == FObjectProperty::StaticClass())
 	{
 		WriteUObject(State, CastFieldChecked<FObjectProperty>(State.ValueProperty)->GetPropertyValue_InContainer(State.ValueData, ArrayIndex));
 	}
@@ -27,7 +27,7 @@ void FMobileJSStructSerializerBackend::WriteUObject(const FStructSerializerState
 {
 	// Note this function uses WriteRawJSONValue to append non-json data to the output stream.
 	FString RawValue = Scripting->ConvertObject(Value);
-	if ((State.ValueProperty == nullptr) || (State.ValueProperty->ArrayDim > 1) || (State.ValueProperty->GetOuter()->GetClass() == FArrayProperty::StaticClass()))
+	if ((State.ValueProperty == nullptr) || (State.ValueProperty->ArrayDim > 1) || (State.ValueProperty->GetOwner<FArrayProperty>() != nullptr))
 	{
 		GetWriter()->WriteRawJSONValue(RawValue);
 	}
