Submitting a modified pull request:
PR #6254: Fix An Incorrect Uint Assignment, Causing Problems On Clang Builds (Contributed by [REDACTED_PERSON])

As the struct is part of a documented debugger interface in VS, I changed the assignement instead of the struct member type.

diff --git a/Engine/Source/Runtime/WebBrowser/Private/WebBrowserSingleton.cpp b/Engine/Source/Runtime/WebBrowser/Private/WebBrowserSingleton.cpp
--- a/Engine/Source/Runtime/WebBrowser/Private/WebBrowserSingleton.cpp
+++ b/Engine/Source/Runtime/WebBrowser/Private/WebBrowserSingleton.cpp
@@ -109,7 +109,7 @@ namespace {
 			uint32 dwFlags;		// Reserved for future use, must be zero.
 		};
 
-		THREADNAME_INFO ThreadNameInfo = {0x1000, ThreadName, -1, 0};
+		THREADNAME_INFO ThreadNameInfo = {0x1000, ThreadName, (uint32)-1, 0};
 
 		__try
 		{
