Updated CEFWebBrowserWindow to prevent issue where browsers used with widget components would cause an "SObject destroyed while collecting garbage" ensure

diff --git a/Engine/Source/Runtime/WebBrowser/Private/CEF/CEFWebBrowserWindow.cpp b/Engine/Source/Runtime/WebBrowser/Private/CEF/CEFWebBrowserWindow.cpp
--- a/Engine/Source/Runtime/WebBrowser/Private/CEF/CEFWebBrowserWindow.cpp
+++ b/Engine/Source/Runtime/WebBrowser/Private/CEF/CEFWebBrowserWindow.cpp
@@ -540,7 +540,11 @@ void FCEFWebBrowserWindow::SetViewportSize(FIntPoint WindowSize, FIntPoint Windo
 	}
 	bTickedLastFrame=true;
 
-	const float WindowDPIScaleFactor = ParentWindow.IsValid() ? ParentWindow->GetNativeWindow()->GetDPIScaleFactor() : 1.0f;
+	float WindowDPIScaleFactor = 1.0f;
+	if (TSharedPtr<SWindow> ParentWindowPtr = ParentWindow.Pin())
+	{
+		WindowDPIScaleFactor = ParentWindowPtr->GetNativeWindow()->GetDPIScaleFactor();
+	}
 
 	// Ignore sizes that can't be seen as it forces CEF to re-render whole image
 	if ((WindowSize.X > 0 && WindowSize.Y > 0 && ViewportSize != WindowSize) || WindowDPIScaleFactor != ViewportDPIScaleFactor)
@@ -2117,7 +2121,12 @@ int32 FCEFWebBrowserWindow::GetCefMouseModifiers(const FPointerEvent& InMouseEve
 CefMouseEvent FCEFWebBrowserWindow::GetCefMouseEvent(const FGeometry& MyGeometry, const FPointerEvent& MouseEvent, bool bIsPopup)
 {
 	CefMouseEvent Event;
-	const float DPIScale = MyGeometry.Scale / (ParentWindow.IsValid() ? ParentWindow->GetNativeWindow()->GetDPIScaleFactor() : 1.0f);
+	float DPIScale = MyGeometry.Scale;
+	if (TSharedPtr<SWindow> ParentWindowPtr = ParentWindow.Pin())
+	{
+		DPIScale /= ParentWindowPtr->GetNativeWindow()->GetDPIScaleFactor();
+	}
+
 	FVector2D LocalPos = MyGeometry.AbsoluteToLocal(MouseEvent.GetScreenSpacePosition()) * DPIScale;
 	if (bIsPopup)
 	{
@@ -2279,9 +2288,10 @@ void FCEFWebBrowserWindow::SetIsHidden(bool bValue)
 
 			if (bIsHidden )
 			{
-				if (::IsWindowEnabled(NativeWindowHandle) && ParentWindow.IsValid())
+				TSharedPtr<SWindow> ParentWindowPtr = ParentWindow.Pin();
+				if (::IsWindowEnabled(NativeWindowHandle) && ParentWindowPtr.IsValid())
 				{
-					::SetFocus((HWND)ParentWindow->GetNativeWindow()->GetOSWindowHandle());
+					::SetFocus((HWND)ParentWindowPtr->GetNativeWindow()->GetOSWindowHandle());
 				}
 			}
 		}
@@ -2301,7 +2311,8 @@ void FCEFWebBrowserWindow::SetIsDisabled(bool bValue)
 
 TSharedPtr<SWindow> FCEFWebBrowserWindow::GetParentWindow() const
 {
-	return ParentWindow;
+	TSharedPtr<SWindow> ParentWindowPtr = ParentWindow.Pin();
+	return ParentWindowPtr;
 }
 
 void FCEFWebBrowserWindow::SetParentWindow(TSharedPtr<SWindow> Window)
diff --git a/Engine/Source/Runtime/WebBrowser/Private/CEF/CEFWebBrowserWindow.h b/Engine/Source/Runtime/WebBrowser/Private/CEF/CEFWebBrowserWindow.h
--- a/Engine/Source/Runtime/WebBrowser/Private/CEF/CEFWebBrowserWindow.h
+++ b/Engine/Source/Runtime/WebBrowser/Private/CEF/CEFWebBrowserWindow.h
@@ -759,7 +759,7 @@ private:
 
 	TArray<FWebBrowserDragRegion> DragRegions;
 
-	TSharedPtr<SWindow> ParentWindow;
+	TWeakPtr<SWindow> ParentWindow;
 };
 
 typedef FCEFWebBrowserWindow FWebBrowserWindow;
