Fix build following 15635368.

#[fyi] alfred.reynolds

diff --git a/Engine/Source/Runtime/WebBrowser/Private/CEF/CEFWebBrowserWindowRHIHelper.h b/Engine/Source/Runtime/WebBrowser/Private/CEF/CEFWebBrowserWindowRHIHelper.h
--- a/Engine/Source/Runtime/WebBrowser/Private/CEF/CEFWebBrowserWindowRHIHelper.h
+++ b/Engine/Source/Runtime/WebBrowser/Private/CEF/CEFWebBrowserWindowRHIHelper.h
@@ -6,7 +6,10 @@
 
 #if WITH_CEF3
 
+#include "Layout/Geometry.h"
+
 class FSlateTexture2DRHIRef;
+class FSlateUpdatableTexture;
 
 /**
  * Implementation of RHI renderer details for the CEF accelerated rendering path
@@ -21,7 +24,7 @@ public:
 
 public:
 	static bool BUseRHIRenderer();
-	FSlateUpdatableTexture*CreateTexture(void *ShareHandle);
+	FSlateUpdatableTexture* CreateTexture(void *ShareHandle);
 	void UpdateSharedHandleTexture(void* SharedHandle, FSlateUpdatableTexture* SlateTexture, const FIntRect& DirtyIn);
 	void UpdateCachedGeometry(const FGeometry& AllottedGeometry);
 	TOptional<FSlateRenderTransform> GetWebBrowserRenderTransform() const;
