Fix another build issue following 15635368.

#[fyi] alfred.reynolds

diff --git a/Engine/Source/Runtime/WebBrowser/Private/CEF/CEFWebBrowserWindowRHIHelper.cpp b/Engine/Source/Runtime/WebBrowser/Private/CEF/CEFWebBrowserWindowRHIHelper.cpp
--- a/Engine/Source/Runtime/WebBrowser/Private/CEF/CEFWebBrowserWindowRHIHelper.cpp
+++ b/Engine/Source/Runtime/WebBrowser/Private/CEF/CEFWebBrowserWindowRHIHelper.cpp
@@ -8,7 +8,7 @@
 #include "CEF/CEFWebBrowserWindow.h"
 #if WITH_ENGINE
 #include "RHI.h"
-#if !PLATFORM_LINUX
+#if PLATFORM_WINDOWS
 #include "Slate/SlateTextures.h"
 #endif
 #endif
