- Add include guard to CEFLibCefIncludes.h and remove defining OVERRIDE macro as it is not needed.

diff --git a/Engine/Source/Runtime/WebBrowser/Private/CEF/CEFLibCefIncludes.h b/Engine/Source/Runtime/WebBrowser/Private/CEF/CEFLibCefIncludes.h
--- a/Engine/Source/Runtime/WebBrowser/Private/CEF/CEFLibCefIncludes.h
+++ b/Engine/Source/Runtime/WebBrowser/Private/CEF/CEFLibCefIncludes.h
@@ -1,11 +1,8 @@
 // Copyright Epic Games, Inc. All Rights Reserved.
+#pragma once
 
 #if WITH_CEF3
 
-#ifndef OVERRIDE
-#	define OVERRIDE override
-#endif //OVERRIDE
-
 #if PLATFORM_WINDOWS
 #	include "Windows/WindowsHWrapper.h"
 #	include "Windows/AllowWindowsPlatformTypes.h"
