- Disable CEF (web browser) accelerated paint by default until we can fix the multi-gpu rendering issue.
   Add "-forcecefaccelpaint" to your command line after this change to turn accelerated rendering back on.

diff --git a/Engine/Source/Runtime/WebBrowser/Private/CEF/CEFWebBrowserWindow.cpp b/Engine/Source/Runtime/WebBrowser/Private/CEF/CEFWebBrowserWindow.cpp
--- a/Engine/Source/Runtime/WebBrowser/Private/CEF/CEFWebBrowserWindow.cpp
+++ b/Engine/Source/Runtime/WebBrowser/Private/CEF/CEFWebBrowserWindow.cpp
@@ -2351,7 +2351,8 @@ bool FCEFWebBrowserWindow::CanSupportAcceleratedPaint()
 	{
 		return false;
 	}
-	return true;
+	// Disable accelerated paint by default until we can fix the multi-gpu rendering issue
+	return false;
 #else
 	return false; // 32-bit windows doesn't have the accelerated rendering patches applied, it can be done if needed
 #endif
