- force a minimum window dimension of 1px when using a WebBrowser in OSR mode. CEF/Chrome isn't compatible with 0 sized windows.

diff --git a/Engine/Source/Runtime/WebBrowser/Private/CEF/CEFWebBrowserWindow.cpp b/Engine/Source/Runtime/WebBrowser/Private/CEF/CEFWebBrowserWindow.cpp
--- a/Engine/Source/Runtime/WebBrowser/Private/CEF/CEFWebBrowserWindow.cpp
+++ b/Engine/Source/Runtime/WebBrowser/Private/CEF/CEFWebBrowserWindow.cpp
@@ -1682,8 +1682,9 @@ void FCEFWebBrowserWindow::GetViewRect(CefRect& Rect)
 	}
 	else
 	{
-		Rect.width = ViewportSize.X;
-		Rect.height = ViewportSize.Y;
+		// CEF requires a minimum of a 1px in each dimension to correctly run
+		Rect.width = FMath::Max( ViewportSize.X, 1);
+		Rect.height = FMath::Max( ViewportSize.Y, 1);
 	}
 }
 
