- Disable accelerated rendering on macOS for the web browser control as we lack RHI support for it right now
- Fix crash when enabling the WebBrowser plugin on macOS. CEF uses a dylib scoped function pointer table for its API, so we need to enable it in the WebBrowser module, doing it in the CEF3Utils module doesn't help here.

diff --git a/Engine/Source/Runtime/WebBrowser/Private/CEF/CEFWebBrowserWindow.cpp b/Engine/Source/Runtime/WebBrowser/Private/CEF/CEFWebBrowserWindow.cpp
--- a/Engine/Source/Runtime/WebBrowser/Private/CEF/CEFWebBrowserWindow.cpp
+++ b/Engine/Source/Runtime/WebBrowser/Private/CEF/CEFWebBrowserWindow.cpp
@@ -2358,7 +2358,7 @@ bool FCEFWebBrowserWindow::CanSupportAcceleratedPaint()
 	return false; // 32-bit windows doesn't have the accelerated rendering patches applied, it can be done if needed
 #endif
 #elif PLATFORM_MAC
-	return true;
+	return false; // Needs RHI support for the CreateSharedHandleTexture call
 #else
 	return false;
 #endif
diff --git a/Engine/Source/Runtime/WebBrowser/Private/WebBrowserModule.cpp b/Engine/Source/Runtime/WebBrowser/Private/WebBrowserModule.cpp
--- a/Engine/Source/Runtime/WebBrowser/Private/WebBrowserModule.cpp
+++ b/Engine/Source/Runtime/WebBrowser/Private/WebBrowserModule.cpp
@@ -7,6 +7,12 @@
 #include "Misc/EngineVersion.h"
 #if WITH_CEF3
 #	include "CEF3Utils.h"
+#	if PLATFORM_MAC
+#		include "include/wrapper/cef_library_loader.h"
+#		define CEF3_BIN_DIR TEXT("Binaries/ThirdParty/CEF3")
+#		define CEF3_FRAMEWORK_DIR CEF3_BIN_DIR TEXT("/Mac/Chromium Embedded Framework.framework")
+#		define CEF3_FRAMEWORK_EXE CEF3_FRAMEWORK_DIR TEXT("/Chromium Embedded Framework")
+#	endif
 #endif
 
 DEFINE_LOG_CATEGORY(LogWebBrowser);
@@ -33,6 +39,10 @@ public:
 private:
 #if WITH_CEF3
 	bool bLoadedCEFModule = false;
+#if PLATFORM_MAC
+	// Dynamically load the CEF framework library.
+	CefScopedLibraryLoader *CEFLibraryLoader = nullptr;
+#endif
 #endif
 };
 
@@ -42,6 +52,20 @@ void FWebBrowserModule::StartupModule()
 {
 #if WITH_CEF3
 	bLoadedCEFModule = CEF3Utils::LoadCEF3Modules(true);
+#if PLATFORM_MAC
+	// Dynamically load the CEF framework library into this dylibs memory space.
+	// CEF now loads function pointers at runtime so we need this to be dylib specific.
+	CEFLibraryLoader = new CefScopedLibraryLoader();
+	
+	FString CefFrameworkPath(FPaths::Combine(*FPaths::EngineDir(), CEF3_FRAMEWORK_EXE));
+	CefFrameworkPath = FPaths::ConvertRelativePathToFull(CefFrameworkPath);
+	
+	bool bLoaderInitialized = false;
+	if (!CEFLibraryLoader->LoadInMain(TCHAR_TO_ANSI(*CefFrameworkPath)))
+	{
+			UE_LOG(LogWebBrowser, Error, TEXT("Chromium loader initialization failed"));
+	}
+#endif // PLATFORM_MAC
 #endif
 }
 
@@ -55,6 +79,10 @@ void FWebBrowserModule::ShutdownModule()
 
 #if WITH_CEF3
 	CEF3Utils::UnloadCEF3Modules();
+#if PLATFORM_MAC
+	delete CEFLibraryLoader;
+	CEFLibraryLoader = nullptr;
+#endif // PLATFORM_MAC
 #endif
 }
 
