CEF
- When setting the parent Slate window for a browser object also set the parent HWND if using direct rendering for the browser control. This lets us support tabs changing the parent HWND as they are docked and undocked

diff --git a/Engine/Source/Runtime/WebBrowser/Private/CEF/CEFWebBrowserWindow.cpp b/Engine/Source/Runtime/WebBrowser/Private/CEF/CEFWebBrowserWindow.cpp
--- a/Engine/Source/Runtime/WebBrowser/Private/CEF/CEFWebBrowserWindow.cpp
+++ b/Engine/Source/Runtime/WebBrowser/Private/CEF/CEFWebBrowserWindow.cpp
@@ -2505,6 +2505,24 @@ TSharedPtr<SWindow> FCEFWebBrowserWindow::GetParentWindow() const
 void FCEFWebBrowserWindow::SetParentWindow(TSharedPtr<SWindow> Window)
 {
 	ParentWindow = Window;
+#if PLATFORM_WINDOWS
+	if (IsValid())
+	{
+		CefRefPtr<CefBrowserHost> BrowserHost = InternalCefBrowser->GetHost();
+
+		HWND NativeWindowHandle = BrowserHost->GetWindowHandle();
+		if (NativeWindowHandle != nullptr)
+		{
+			TSharedPtr<SWindow> ParentWindowPtr = ParentWindow.Pin();
+			void* ParentWindowHandle = (ParentWindow.IsValid() && ParentWindowPtr->GetNativeWindow().IsValid()) ? ParentWindowPtr->GetNativeWindow()->GetOSWindowHandle() : nullptr;
+			if (ParentWindowHandle != nullptr)
+			{
+				// When rendering directly to a HWND update its parent windown
+				::SetParent(NativeWindowHandle, (HWND)ParentWindowHandle);
+			}
+		}
+	}
+#endif
 }
 
 CefRefPtr<CefDictionaryValue> FCEFWebBrowserWindow::GetProcessInfo()
diff --git a/Engine/Source/Runtime/WebBrowser/Private/SWebBrowser.cpp b/Engine/Source/Runtime/WebBrowser/Private/SWebBrowser.cpp
--- a/Engine/Source/Runtime/WebBrowser/Private/SWebBrowser.cpp
+++ b/Engine/Source/Runtime/WebBrowser/Private/SWebBrowser.cpp
@@ -386,4 +386,13 @@ void SWebBrowser::UnbindInputMethodSystem()
 	}
 }
 
+void SWebBrowser::SetParentWindow(TSharedPtr<SWindow> Window)
+{
+	if (BrowserView.IsValid())
+	{
+		BrowserView->SetParentWindow(Window);
+	}
+}
+
+
 #undef LOCTEXT_NAMESPACE
diff --git a/Engine/Source/Runtime/WebBrowser/Private/SWebBrowserView.cpp b/Engine/Source/Runtime/WebBrowser/Private/SWebBrowserView.cpp
--- a/Engine/Source/Runtime/WebBrowser/Private/SWebBrowserView.cpp
+++ b/Engine/Source/Runtime/WebBrowser/Private/SWebBrowserView.cpp
@@ -720,4 +720,14 @@ bool SWebBrowserView::UnhandledKeyChar(const FCharacterEvent& CharacterEvent)
 }
 
 
+void SWebBrowserView::SetParentWindow(TSharedPtr<SWindow> Window)
+{
+	SetupParentWindowHandlers();
+	if (BrowserWindow.IsValid())
+	{
+		BrowserWindow->SetParentWindow(Window);
+	}
+}
+
+
 #undef LOCTEXT_NAMESPACE
diff --git a/Engine/Source/Runtime/WebBrowser/Public/SWebBrowser.h b/Engine/Source/Runtime/WebBrowser/Public/SWebBrowser.h
--- a/Engine/Source/Runtime/WebBrowser/Public/SWebBrowser.h
+++ b/Engine/Source/Runtime/WebBrowser/Public/SWebBrowser.h
@@ -235,6 +235,9 @@ public:
 	/** Navigate forwards. */
 	void GoForward();
 
+	/** Set parent SWindow for this browser. */
+	void SetParentWindow(TSharedPtr<SWindow> Window);
+
 private:
 
 	/** Navigate backwards. */
diff --git a/Engine/Source/Runtime/WebBrowser/Public/SWebBrowserView.h b/Engine/Source/Runtime/WebBrowser/Public/SWebBrowserView.h
--- a/Engine/Source/Runtime/WebBrowser/Public/SWebBrowserView.h
+++ b/Engine/Source/Runtime/WebBrowser/Public/SWebBrowserView.h
@@ -266,6 +266,8 @@ public:
 	/** Navigate forwards. */
 	void GoForward();
 
+	/** Set parent SWindow for this browser. */
+	void SetParentWindow(TSharedPtr<SWindow> Window);
 private:
 
 	void SetupParentWindowHandlers();
