Fix crash when attempting to convert a level to world partition

Added a flag (AllowCommandletCEF) for enabling CEF in commandlet mode

diff --git a/Engine/Source/Runtime/WebBrowser/Private/WebBrowserSingleton.cpp b/Engine/Source/Runtime/WebBrowser/Private/WebBrowserSingleton.cpp
--- a/Engine/Source/Runtime/WebBrowser/Private/WebBrowserSingleton.cpp
+++ b/Engine/Source/Runtime/WebBrowser/Private/WebBrowserSingleton.cpp
@@ -235,7 +235,7 @@ FWebBrowserSingleton::FWebBrowserSingleton(const FWebBrowserInitSettings& WebBro
 #if WITH_CEF3
 	
 	// Only enable CEF if we have CEF3, we are not running a commandlet without rendering (e.g. cooking assets) and it has not been explicitly disabled
-	bAllowCEF = (!IsRunningCommandlet() || IsAllowCommandletRendering()) && !FParse::Param(FCommandLine::Get(), TEXT("nocef"));
+	bAllowCEF = (!IsRunningCommandlet() || IsAllowCommandletRendering()) && !FParse::Param(FCommandLine::Get(), TEXT("nocef")) && FParse::Param(FCommandLine::Get(), TEXT("AllowCommandletCEF"));
 	if (bAllowCEF)
 	{
 		// The FWebBrowserSingleton must be initialized on the game thread
