Fix the game thread sometimes showing up as 'CrBrowserMain' in the debugger.

#rnx

- When the WebBrowser module is loaded in the editor we make a call to ::CefInitialize which seems to set the current thread name to 'CrBrowserMain' overriding that the engine sets on start up.
- This was originally fixed back in CL 2636308 for UE-5165 but uses it's own utility function to set the thread name rather than using FPlatformProcess which contains more up to date ways of setting the thread name on certain platforms.
- Removed the local utility method and started using FPlatformProcess::SetThreadName.

diff --git a/Engine/Source/Runtime/WebBrowser/Private/WebBrowserSingleton.cpp b/Engine/Source/Runtime/WebBrowser/Private/WebBrowserSingleton.cpp
--- a/Engine/Source/Runtime/WebBrowser/Private/WebBrowserSingleton.cpp
+++ b/Engine/Source/Runtime/WebBrowser/Private/WebBrowserSingleton.cpp
@@ -84,49 +84,6 @@ THIRD_PARTY_INCLUDES_END
 #	endif
 #endif
 
-namespace {
-
-	/**
-	 * Helper function to set the current thread name, visible by the debugger.
-	 * @param ThreadName	Name to set
-	 */
-	void SetCurrentThreadName(char* ThreadName)
-	{
-#if PLATFORM_MAC
-		pthread_setname_np(ThreadName);
-#elif PLATFORM_LINUX
-		pthread_setname_np(pthread_self(), ThreadName);
-#elif PLATFORM_WINDOWS && !PLATFORM_SEH_EXCEPTIONS_DISABLED
-		/**
-		 * Code setting the thread name for use in the debugger.
-		 * Copied implementation from WindowsRunnableThread as it is private.
-		 *
-		 * http://msdn.microsoft.com/en-us/library/xcb2z8hs.aspx
-		 */
-		const uint32 MS_VC_EXCEPTION=0x406D1388;
-
-		struct THREADNAME_INFO
-		{
-			uint32 dwType;		// Must be 0x1000.
-			LPCSTR szName;		// Pointer to name (in user addr space).
-			uint32 dwThreadID;	// Thread ID (-1=caller thread).
-			uint32 dwFlags;		// Reserved for future use, must be zero.
-		};
-
-		THREADNAME_INFO ThreadNameInfo = {0x1000, ThreadName, (uint32)-1, 0};
-
-		__try
-		{
-			RaiseException( MS_VC_EXCEPTION, 0, sizeof(ThreadNameInfo)/sizeof(ULONG_PTR), (ULONG_PTR*)&ThreadNameInfo );
-		}
-		__except( EXCEPTION_EXECUTE_HANDLER )
-		CA_SUPPRESS(6322)
-		{
-		}
-#endif
-	}
-}
-
 FString FWebBrowserSingleton::ApplicationCacheDir() const
 {
 #if PLATFORM_MAC
@@ -349,7 +306,7 @@ FWebBrowserSingleton::FWebBrowserSingleton(const FWebBrowserInitSettings& WebBro
 		check(bSuccess);
 
 		// Set the thread name back to GameThread.
-		SetCurrentThreadName(TCHAR_TO_ANSI( *(FName( NAME_GameThread ).GetPlainNameString()) ));
+		FPlatformProcess::SetThreadName(*FName(NAME_GameThread).GetPlainNameString());
 
 		DefaultCookieManager = FCefWebBrowserCookieManagerFactory::Create(CefCookieManager::GetGlobalManager(nullptr));
 	}
