With CEF3 we need to save then restore our signal handlers when init'ing CEF

Cherry picking CL 15923539 to UE5

diff --git a/Engine/Source/Runtime/WebBrowser/Private/WebBrowserSingleton.cpp b/Engine/Source/Runtime/WebBrowser/Private/WebBrowserSingleton.cpp
--- a/Engine/Source/Runtime/WebBrowser/Private/WebBrowserSingleton.cpp
+++ b/Engine/Source/Runtime/WebBrowser/Private/WebBrowserSingleton.cpp
@@ -178,6 +178,49 @@ public:
 	}
 };
 
+#if WITH_CEF3
+#if PLATFORM_MAC || PLATFORM_LINUX
+class FPosixSignalPreserver
+{
+public:
+	FPosixSignalPreserver()
+	{
+		struct sigaction Sigact;
+		for (uint32 i = 0; i < UE_ARRAY_COUNT(PreserveSignals); ++i)
+		{
+			FMemory::Memset(&Sigact, 0, sizeof(Sigact));
+			if (sigaction(PreserveSignals[i], nullptr, &Sigact) != 0)
+			{
+				UE_LOG(LogWebBrowser, Warning, TEXT("Failed to backup signal handler for %i."), PreserveSignals[i]);
+			}
+			OriginalSignalHandlers[i] = Sigact;
+		}
+	}
+
+	~FPosixSignalPreserver()
+	{
+		for (uint32 i = 0; i < UE_ARRAY_COUNT(PreserveSignals); ++i)
+		{
+			if(sigaction(PreserveSignals[i], &OriginalSignalHandlers[i], nullptr) != 0)
+			{
+				UE_LOG(LogWebBrowser, Warning, TEXT("Failed to restore signal handler for %i."), PreserveSignals[i]);
+			}
+		}
+	}
+
+private:
+	// Backup the list of signals that CEF/Chromium overrides, derived from SetupSignalHandlers() in
+	//  https://chromium.googlesource.com/chromium/src.git/+/2fc330d0b93d4bfd7bd04b9fdd3102e529901f91/services/service_manager/embedder/main.cc
+	const int PreserveSignals[13] = {SIGHUP, SIGINT, SIGQUIT, SIGILL, SIGABRT,
+		SIGFPE, SIGSEGV, SIGALRM, SIGTERM, SIGCHLD, SIGBUS, SIGTRAP, SIGPIPE};
+
+	struct sigaction OriginalSignalHandlers[UE_ARRAY_COUNT(PreserveSignals)];
+};
+
+#endif // PLATFORM_MAC || PLATFORM_LINUX
+#endif // WITH_CEF3
+
+
 FWebBrowserSingleton::FWebBrowserSingleton(const FWebBrowserInitSettings& WebBrowserInitSettings)
 #if WITH_CEF3
 	: WebBrowserWindowFactory(MakeShareable(new FWebBrowserWindowFactory()))
@@ -301,6 +344,11 @@ FWebBrowserSingleton::FWebBrowserSingleton(const FWebBrowserInitSettings& WebBro
 		}
 		CefString(&Settings.browser_subprocess_path) = TCHAR_TO_WCHAR(*SubProcessPath);
 
+#if PLATFORM_MAC || PLATFORM_LINUX
+		// this class automatically preserves the sigaction handlers we have set
+		FPosixSignalPreserver PosixSignalPreserver;
+#endif
+
 		// Initialize CEF.
 		bool bSuccess = CefInitialize(MainArgs, Settings, CEFBrowserApp.get(), nullptr);
 		check(bSuccess);
