Remove reference to UE4 in comments.

diff --git a/Engine/Source/Runtime/WebBrowser/Public/IWebBrowserWindow.h b/Engine/Source/Runtime/WebBrowser/Public/IWebBrowserWindow.h
--- a/Engine/Source/Runtime/WebBrowser/Public/IWebBrowserWindow.h
+++ b/Engine/Source/Runtime/WebBrowser/Public/IWebBrowserWindow.h
@@ -275,7 +275,7 @@ public:
 	 * Properties and Functions will be accessible from JavaScript side.
 	 * As all communication with the rendering procesis asynchronous, return values (both for properties and function results) are wrapped into JS Future objects.
 	 *
-	 * @param Name The name of the object. The object will show up as window.ue4.{Name} on the javascript side. If there is an existing object of the same name, this object will replace it. If bIsPermanent is false and there is an existing permanent binding, the permanent binding will be restored when the temporary one is removed.
+	 * @param Name The name of the object. The object will show up as window.ue.{Name} on the javascript side. If there is an existing object of the same name, this object will replace it. If bIsPermanent is false and there is an existing permanent binding, the permanent binding will be restored when the temporary one is removed.
 	 * @param Object The object instance.
 	 * @param bIsPermanent If true, the object will be visible to all pages loaded through this browser widget, otherwise, it will be deleted when navigating away from the current page. Non-permanent bindings should be registered from inside an OnLoadStarted event handler in order to be available before JS code starts loading.
 	 */
diff --git a/Engine/Source/Runtime/WebBrowser/Public/SWebBrowser.h b/Engine/Source/Runtime/WebBrowser/Public/SWebBrowser.h
--- a/Engine/Source/Runtime/WebBrowser/Public/SWebBrowser.h
+++ b/Engine/Source/Runtime/WebBrowser/Public/SWebBrowser.h
@@ -200,7 +200,7 @@ public:
 	 * Properties and Functions will be accessible from JavaScript side.
 	 * As all communication with the rendering procesis asynchronous, return values (both for properties and function results) are wrapped into JS Future objects.
 	 *
-	 * @param Name The name of the object. The object will show up as window.ue4.{Name} on the javascript side. If there is an existing object of the same name, this object will replace it. If bIsPermanent is false and there is an existing permanent binding, the permanent binding will be restored when the temporary one is removed.
+	 * @param Name The name of the object. The object will show up as window.ue.{Name} on the javascript side. If there is an existing object of the same name, this object will replace it. If bIsPermanent is false and there is an existing permanent binding, the permanent binding will be restored when the temporary one is removed.
 	 * @param Object The object instance.
 	 * @param bIsPermanent If true, the object will be visible to all pages loaded through this browser widget, otherwise, it will be deleted when navigating away from the current page. Non-permanent bindings should be registered from inside an OnLoadStarted event handler in order to be available before JS code starts loading.
 	 */
diff --git a/Engine/Source/Runtime/WebBrowser/Public/SWebBrowserView.h b/Engine/Source/Runtime/WebBrowser/Public/SWebBrowserView.h
--- a/Engine/Source/Runtime/WebBrowser/Public/SWebBrowserView.h
+++ b/Engine/Source/Runtime/WebBrowser/Public/SWebBrowserView.h
@@ -231,7 +231,7 @@ public:
 	 * Properties and Functions will be accessible from JavaScript side.
 	 * As all communication with the rendering procesis asynchronous, return values (both for properties and function results) are wrapped into JS Future objects.
 	 *
-	 * @param Name The name of the object. The object will show up as window.ue4.{Name} on the javascript side. If there is an existing object of the same name, this object will replace it. If bIsPermanent is false and there is an existing permanent binding, the permanent binding will be restored when the temporary one is removed.
+	 * @param Name The name of the object. The object will show up as window.ue.{Name} on the javascript side. If there is an existing object of the same name, this object will replace it. If bIsPermanent is false and there is an existing permanent binding, the permanent binding will be restored when the temporary one is removed.
 	 * @param Object The object instance.
 	 * @param bIsPermanent If true, the object will be visible to all pages loaded through this browser widget, otherwise, it will be deleted when navigating away from the current page. Non-permanent bindings should be registered from inside an OnLoadStarted event handler in order to be available before JS code starts loading.
 	 */
