Disable CEF if we are running with -RenderOffscreen
Disable CEF if we cannot ever render

diff --git a/Engine/Source/Runtime/WebBrowser/Private/WebBrowserSingleton.cpp b/Engine/Source/Runtime/WebBrowser/Private/WebBrowserSingleton.cpp
--- a/Engine/Source/Runtime/WebBrowser/Private/WebBrowserSingleton.cpp
+++ b/Engine/Source/Runtime/WebBrowser/Private/WebBrowserSingleton.cpp
@@ -235,7 +235,10 @@ FWebBrowserSingleton::FWebBrowserSingleton(const FWebBrowserInitSettings& WebBro
 #if WITH_CEF3
 	
 	// Only enable CEF if we have CEF3, we are not running a commandlet without rendering (e.g. cooking assets) and it has not been explicitly disabled
-	bAllowCEF = (!IsRunningCommandlet() || (IsAllowCommandletRendering() && FParse::Param(FCommandLine::Get(), TEXT("AllowCommandletCEF")))) && !FParse::Param(FCommandLine::Get(), TEXT("nocef"));
+	// Disallow CEF if we are running with -RenderOffscreen, many times this includes running with out access to X11 for example on Linux which is required by CEF
+	// Disallow CEF if we never plan on rendering, ie, with CanEverRender. This includes servers
+	bAllowCEF = (!IsRunningCommandlet() || (IsAllowCommandletRendering() && FParse::Param(FCommandLine::Get(), TEXT("AllowCommandletCEF")))) &&
+				!FParse::Param(FCommandLine::Get(), TEXT("RenderOffscreen")) && FApp::CanEverRender() && !FParse::Param(FCommandLine::Get(), TEXT("nocef"));
 	if (bAllowCEF)
 	{
 		// The FWebBrowserSingleton must be initialized on the game thread
