Workaround for a CEF issue related to keyboard focus.
- Add SetBrowserKeyboardFocus() function to SWebBrowserView. Calling this lets you manually synchronise the Slate keyboard focus state with the underlying browser widget.
- To work around CEF 2040 where focus state is lost on page navigation for the OSR renderer in some cases, call SetBrowserKeyboardFocus() when navigating to a new URL or creating iframes. The CEF team reports that v95 and above of CEF have a fix for this issue so this workaround may be removed once upgraded past that version.

#tests Ran on PC and tested with various 3rd party payment flows.

diff --git a/Engine/Source/Runtime/WebBrowser/Private/SWebBrowserView.cpp b/Engine/Source/Runtime/WebBrowser/Private/SWebBrowserView.cpp
--- a/Engine/Source/Runtime/WebBrowser/Private/SWebBrowserView.cpp
+++ b/Engine/Source/Runtime/WebBrowser/Private/SWebBrowserView.cpp
@@ -733,5 +733,9 @@ void SWebBrowserView::SetParentWindow(TSharedPtr<SWindow> Window)
 	}
 }
 
+void SWebBrowserView::SetBrowserKeyboardFocus()
+{
+	BrowserWindow->OnFocus(HasAnyUserFocusOrFocusedDescendants(), false);
+}
 
 #undef LOCTEXT_NAMESPACE
diff --git a/Engine/Source/Runtime/WebBrowser/Public/SWebBrowserView.h b/Engine/Source/Runtime/WebBrowser/Public/SWebBrowserView.h
--- a/Engine/Source/Runtime/WebBrowser/Public/SWebBrowserView.h
+++ b/Engine/Source/Runtime/WebBrowser/Public/SWebBrowserView.h
@@ -268,6 +268,11 @@ public:
 
 	/** Set parent SWindow for this browser. */
 	void SetParentWindow(TSharedPtr<SWindow> Window);
+
+	/** Update the underlying browser widget to match the KB focus in slate.
+		This is used to work around a CEF bug that loses focus state on navigations*/
+	void SetBrowserKeyboardFocus();
+
 private:
 
 	void SetupParentWindowHandlers();
