Fix CIS errors on Mac

diff --git a/Engine/Source/Developer/TraceAnalysis/Private/Store/AsioStore.cpp b/Engine/Source/Developer/TraceAnalysis/Private/Store/AsioStore.cpp
--- a/Engine/Source/Developer/TraceAnalysis/Private/Store/AsioStore.cpp
+++ b/Engine/Source/Developer/TraceAnalysis/Private/Store/AsioStore.cpp
@@ -13,6 +13,8 @@
 #	include "Windows/HideWindowsPlatformTypes.h"
 #elif PLATFORM_LINUX
 	#include <sys/inotify.h>
+#elif PLATFORM_MAC
+	#include "Misc/Paths.h"
 #endif
 
 namespace UE {
diff --git a/Engine/Source/Runtime/WebBrowser/Private/WebBrowserModule.cpp b/Engine/Source/Runtime/WebBrowser/Private/WebBrowserModule.cpp
--- a/Engine/Source/Runtime/WebBrowser/Private/WebBrowserModule.cpp
+++ b/Engine/Source/Runtime/WebBrowser/Private/WebBrowserModule.cpp
@@ -5,6 +5,9 @@
 #include "WebBrowserSingleton.h"
 #include "Misc/App.h"
 #include "Misc/EngineVersion.h"
+#if PLATFORM_MAC
+	#include "Misc/Paths.h"
+#endif
 #if WITH_CEF3
 #	include "CEF3Utils.h"
 #	if PLATFORM_MAC
