Remove NDA platform references from WebBrowser
#rnx

diff --git a/Engine/Source/Runtime/WebBrowser/Private/SWebBrowserView.cpp b/Engine/Source/Runtime/WebBrowser/Private/SWebBrowserView.cpp
--- a/Engine/Source/Runtime/WebBrowser/Private/SWebBrowserView.cpp
+++ b/Engine/Source/Runtime/WebBrowser/Private/SWebBrowserView.cpp
@@ -17,8 +17,8 @@
 #	include "Android/AndroidWebBrowserWindow.h"
 #elif PLATFORM_IOS
 #	include "IOS/IOSPlatformWebBrowser.h"
-#elif PLATFORM_PS4
-#	include "PS4PlatformWebBrowser.h"
+#elif PLATFORM_SPECIFIC_WEB_BROWSER
+#	include COMPILED_PLATFORM_HEADER(PlatformWebBrowser.h)
 #elif WITH_CEF3
 #	include "CEF/CEFWebBrowserWindow.h"
 #else
diff --git a/Engine/Source/Runtime/WebBrowser/Private/WebBrowserSingleton.cpp b/Engine/Source/Runtime/WebBrowser/Private/WebBrowserSingleton.cpp
--- a/Engine/Source/Runtime/WebBrowser/Private/WebBrowserSingleton.cpp
+++ b/Engine/Source/Runtime/WebBrowser/Private/WebBrowserSingleton.cpp
@@ -57,8 +57,8 @@ THIRD_PARTY_INCLUDES_END
 #elif PLATFORM_IOS
 #	include <IOS/IOSPlatformWebBrowser.h>
 #	include <IOS/IOSCookieManager.h>
-#elif PLATFORM_PS4
-#	include "PS4PlatformWebBrowser.h"
+#elif PLATFORM_SPECIFIC_WEB_BROWSER
+#	include COMPILED_PLATFORM_HEADER(PlatformWebBrowser.h)
 #endif
 
 // Define some platform-dependent file locations
diff --git a/Engine/Source/Runtime/WebBrowser/WebBrowser.Build.cs b/Engine/Source/Runtime/WebBrowser/WebBrowser.Build.cs
--- a/Engine/Source/Runtime/WebBrowser/WebBrowser.Build.cs
+++ b/Engine/Source/Runtime/WebBrowser/WebBrowser.Build.cs
@@ -91,5 +91,7 @@ public class WebBrowser : ModuleRules
 				}
 			}
 		}
+		PrivateDefinitions.Add("PLATFORM_SPECIFIC_WEB_BROWSER=" + (bPlatformSpecificWebBrowser ? "1" : "0"));
 	}
+	protected virtual bool bPlatformSpecificWebBrowser { get { return false; } }
 }
