- Fixed some platforms after recent CEF changes that didn't come in from other stream

diff --git a/Engine/Source/Runtime/WebBrowser/Private/Android/AndroidWebBrowserWindow.cpp b/Engine/Source/Runtime/WebBrowser/Private/Android/AndroidWebBrowserWindow.cpp
--- a/Engine/Source/Runtime/WebBrowser/Private/Android/AndroidWebBrowserWindow.cpp
+++ b/Engine/Source/Runtime/WebBrowser/Private/Android/AndroidWebBrowserWindow.cpp
@@ -38,7 +38,7 @@ FAndroidWebBrowserWindow::FAndroidWebBrowserWindow(FString InUrl, TOptional<FStr
 
 FAndroidWebBrowserWindow::~FAndroidWebBrowserWindow()
 {
-	CloseBrowser(true);
+	CloseBrowser(true, false);
 }
 
 void FAndroidWebBrowserWindow::LoadURL(FString NewURL)
@@ -164,6 +164,11 @@ FReply FAndroidWebBrowserWindow::OnMouseWheel(const FGeometry& MyGeometry, const
 	return FReply::Unhandled();
 }
 
+FReply FAndroidWebBrowserWindow::OnTouchGesture(const FGeometry& MyGeometry, const FPointerEvent& GestureEvent, bool bIsPopup)
+{
+	return FReply::Unhandled();
+}
+
 void FAndroidWebBrowserWindow::OnFocus(bool SetFocus, bool bIsPopup)
 {
 }
@@ -326,4 +331,4 @@ bool FAndroidWebBrowserWindow::IsVisible()
 	return bIsVisible;
 }
 
-#endif // USE_ANDROID_JNI
\ No newline at end of file
+#endif // USE_ANDROID_JNI
diff --git a/Engine/Source/Runtime/WebBrowser/Private/Android/AndroidWebBrowserWindow.h b/Engine/Source/Runtime/WebBrowser/Private/Android/AndroidWebBrowserWindow.h
--- a/Engine/Source/Runtime/WebBrowser/Private/Android/AndroidWebBrowserWindow.h
+++ b/Engine/Source/Runtime/WebBrowser/Private/Android/AndroidWebBrowserWindow.h
@@ -75,6 +75,7 @@ public:
 	virtual void SetSupportsMouseWheel(bool bValue) override;
 	virtual bool GetSupportsMouseWheel() const override;
 	virtual FReply OnMouseWheel(const FGeometry& MyGeometry, const FPointerEvent& MouseEvent, bool bIsPopup) override;
+	virtual FReply OnTouchGesture(const FGeometry& MyGeometry, const FPointerEvent& GestureEvent, bool bIsPopup) override;
 	virtual void OnFocus(bool SetFocus, bool bIsPopup) override;
 	virtual void OnCaptureLost() override;
 	virtual bool CanGoBack() const override;
@@ -349,4 +350,4 @@ private:
 
 typedef FAndroidWebBrowserWindow FWebBrowserWindow;
 
-#endif // USE_ANDROID_JNI
\ No newline at end of file
+#endif // USE_ANDROID_JNI
diff --git a/Engine/Source/Runtime/WebBrowser/Private/IOS/IOSPlatformWebBrowser.cpp b/Engine/Source/Runtime/WebBrowser/Private/IOS/IOSPlatformWebBrowser.cpp
--- a/Engine/Source/Runtime/WebBrowser/Private/IOS/IOSPlatformWebBrowser.cpp
+++ b/Engine/Source/Runtime/WebBrowser/Private/IOS/IOSPlatformWebBrowser.cpp
@@ -830,7 +830,7 @@ FWebBrowserWindow::FWebBrowserWindow(FString InUrl, TOptional<FString> InContent
 
 FWebBrowserWindow::~FWebBrowserWindow()
 {
-	CloseBrowser(true);
+	CloseBrowser(true, false);
 }
 
 void FWebBrowserWindow::LoadURL(FString NewURL)
@@ -957,6 +957,12 @@ FReply FWebBrowserWindow::OnMouseWheel(const FGeometry& MyGeometry, const FPoint
 	return FReply::Unhandled();
 }
 
+FReply FWebBrowserWindow::OnTouchGesture(const FGeometry& MyGeometry, const FPointerEvent& GestureEvent, bool bIsPopup)
+{
+	return FReply::Unhandled();
+}
+
+
 void FWebBrowserWindow::OnFocus(bool SetFocus, bool bIsPopup)
 {
 }
diff --git a/Engine/Source/Runtime/WebBrowser/Private/IOS/IOSPlatformWebBrowser.h b/Engine/Source/Runtime/WebBrowser/Private/IOS/IOSPlatformWebBrowser.h
--- a/Engine/Source/Runtime/WebBrowser/Private/IOS/IOSPlatformWebBrowser.h
+++ b/Engine/Source/Runtime/WebBrowser/Private/IOS/IOSPlatformWebBrowser.h
@@ -133,6 +133,7 @@ public:
 	virtual void SetSupportsMouseWheel(bool bValue) override;
 	virtual bool GetSupportsMouseWheel() const override;
 	virtual FReply OnMouseWheel(const FGeometry& MyGeometry, const FPointerEvent& MouseEvent, bool bIsPopup) override;
+	virtual FReply OnTouchGesture(const FGeometry& MyGeometry, const FPointerEvent& GestureEvent, bool bIsPopup) override;
 	virtual void OnFocus(bool SetFocus, bool bIsPopup) override;
 	virtual void OnCaptureLost() override;
 	virtual bool CanGoBack() const override;
