- Fixed a non-unity error

diff --git a/Engine/Source/Runtime/WebBrowser/Private/WebBrowserModule.cpp b/Engine/Source/Runtime/WebBrowser/Private/WebBrowserModule.cpp
--- a/Engine/Source/Runtime/WebBrowser/Private/WebBrowserModule.cpp
+++ b/Engine/Source/Runtime/WebBrowser/Private/WebBrowserModule.cpp
@@ -5,9 +5,7 @@
 #include "WebBrowserSingleton.h"
 #include "Misc/App.h"
 #include "Misc/EngineVersion.h"
-#if PLATFORM_MAC
-	#include "Misc/Paths.h"
-#endif
+#include "Misc/Paths.h"
 #if WITH_CEF3
 #	include "CEF3Utils.h"
 #	if PLATFORM_MAC
