- Fixed some unacceptable words from CEF integration, and added a known case to the script for a third party #include we can't change

diff --git a/Engine/Source/Runtime/WebBrowser/Private/CEF/CEFBrowserHandler.cpp b/Engine/Source/Runtime/WebBrowser/Private/CEF/CEFBrowserHandler.cpp
--- a/Engine/Source/Runtime/WebBrowser/Private/CEF/CEFBrowserHandler.cpp
+++ b/Engine/Source/Runtime/WebBrowser/Private/CEF/CEFBrowserHandler.cpp
@@ -41,12 +41,12 @@ void LogCEFLoad(const FString &Msg, CefRefPtr<CefRequest> Request)
 // Used to force returning custom content instead of performing a request.
 const FString CustomContentMethod(TEXT("X-GET-CUSTOM-CONTENT"));
 
-FCEFBrowserHandler::FCEFBrowserHandler(bool InUseTransparency, bool InInterceptLoadRequests, const TArray<FString>& InAltRetryDomains, const TArray<FString>& InAuthorizationHeaderWhitelistURLS)
+FCEFBrowserHandler::FCEFBrowserHandler(bool InUseTransparency, bool InInterceptLoadRequests, const TArray<FString>& InAltRetryDomains, const TArray<FString>& InAuthorizationHeaderAllowListURLS)
 : bUseTransparency(InUseTransparency), 
 bAllowAllCookies(false),
 bInterceptLoadRequests(InInterceptLoadRequests),
 AltRetryDomains(InAltRetryDomains),
-AuthorizationHeaderWhitelistURLS(InAuthorizationHeaderWhitelistURLS)
+AuthorizationHeaderAllowListURLS(InAuthorizationHeaderAllowListURLS)
 {
 	// should we forcefully allow all cookies to be set rather than filtering a couple store side ones
 	bAllowAllCookies = FParse::Param(FCommandLine::Get(), TEXT("CefAllowAllCookies"));
@@ -899,10 +899,10 @@ bool FCEFBrowserHandler::URLRequestAllowsCredentials(const FString& URL) const
 	if (MainFrameLoadTypes.Find(URL) != nullptr)
 		return true;
 
-	// check the explicit whitelist also
-	for (const FString& AuthorizationHeaderWhitelistURL : AuthorizationHeaderWhitelistURLS)
+	// check the explicit allowlist also
+	for (const FString& AuthorizationHeaderAllowListURL : AuthorizationHeaderAllowListURLS)
 	{
-		if (URL.Contains(AuthorizationHeaderWhitelistURL))
+		if (URL.Contains(AuthorizationHeaderAllowListURL))
 		{
 			return true;
 		}
diff --git a/Engine/Source/Runtime/WebBrowser/Private/CEF/CEFBrowserHandler.h b/Engine/Source/Runtime/WebBrowser/Private/CEF/CEFBrowserHandler.h
--- a/Engine/Source/Runtime/WebBrowser/Private/CEF/CEFBrowserHandler.h
+++ b/Engine/Source/Runtime/WebBrowser/Private/CEF/CEFBrowserHandler.h
@@ -42,7 +42,7 @@ class FCEFBrowserHandler
 public:
 
 	/** Default constructor. */
-	FCEFBrowserHandler(bool InUseTransparency, bool InInterceptLoadRequests, const TArray<FString>& AltRetryDomains = TArray<FString>(), const TArray<FString>& AuthorizationHeaderWhitelistURLS = TArray<FString>());
+	FCEFBrowserHandler(bool InUseTransparency, bool InInterceptLoadRequests, const TArray<FString>& AltRetryDomains = TArray<FString>(), const TArray<FString>& AuthorizationHeaderAllowListURLS = TArray<FString>());
 
 public:
 
@@ -362,7 +362,7 @@ private:
 	uint32 AltRetryDomainIdx = 0;
 
 	/** Domains we allow sending an authorization header too even if the request doesn't otherwise indicate support */
-	TArray<FString> AuthorizationHeaderWhitelistURLS;
+	TArray<FString> AuthorizationHeaderAllowListURLS;
 
 	/** Keep track of URLs we see being loaded and the type of load it is*/
 	TMap<FString, CefRequest::ResourceType> MainFrameLoadTypes;
diff --git a/Engine/Source/Runtime/WebBrowser/Private/WebBrowserSingleton.cpp b/Engine/Source/Runtime/WebBrowser/Private/WebBrowserSingleton.cpp
--- a/Engine/Source/Runtime/WebBrowser/Private/WebBrowserSingleton.cpp
+++ b/Engine/Source/Runtime/WebBrowser/Private/WebBrowserSingleton.cpp
@@ -552,11 +552,11 @@ TSharedPtr<IWebBrowserWindow> FWebBrowserSingleton::CreateBrowserWindow(const FC
 			BrowserSettings.windowless_frame_rate = BrowserFrameRate;
 		}
 
-		TArray<FString> AuthorizationHeaderWhitelistURLS;
-		GConfig->GetArray(TEXT("Browser"), TEXT("AuthorizationHeaderWhitelistURLS"), AuthorizationHeaderWhitelistURLS, GEngineIni);
+		TArray<FString> AuthorizationHeaderAllowListURLS;
+		GConfig->GetArray(TEXT("Browser"), TEXT("AuthorizationHeaderAllowListURLS"), AuthorizationHeaderAllowListURLS, GEngineIni);
 
 		// WebBrowserHandler implements browser-level callbacks.
-		CefRefPtr<FCEFBrowserHandler> NewHandler(new FCEFBrowserHandler(WindowSettings.bUseTransparency, WindowSettings.bInterceptLoadRequests ,WindowSettings.AltRetryDomains, AuthorizationHeaderWhitelistURLS));
+		CefRefPtr<FCEFBrowserHandler> NewHandler(new FCEFBrowserHandler(WindowSettings.bUseTransparency, WindowSettings.bInterceptLoadRequests ,WindowSettings.AltRetryDomains, AuthorizationHeaderAllowListURLS));
 
 		CefRefPtr<CefRequestContext> RequestContext = nullptr;
 		if (WindowSettings.Context.IsSet())
