Merge CEF changes from //Portal/Main to //UE5/Main/...
- Added missing OnConsoleMessage delegate with this merge
- Added support for direct window rendering under macOS

diff --git a/Engine/Plugins/Runtime/WebBrowserWidget/Source/WebBrowserWidget/Private/WebBrowser.cpp b/Engine/Plugins/Runtime/WebBrowserWidget/Source/WebBrowserWidget/Private/WebBrowser.cpp
--- a/Engine/Plugins/Runtime/WebBrowserWidget/Source/WebBrowserWidget/Private/WebBrowser.cpp
+++ b/Engine/Plugins/Runtime/WebBrowserWidget/Source/WebBrowserWidget/Private/WebBrowser.cpp
@@ -1,4 +1,4 @@
-﻿// Copyright Epic Games, Inc. All Rights Reserved.
+// Copyright Epic Games, Inc. All Rights Reserved.
 
 #include "WebBrowser.h"
 #include "SWebBrowser.h"
@@ -99,7 +99,8 @@ TSharedRef<SWidget> UWebBrowser::RebuildWidget()
 			.ShowControls(false)
 			.SupportsTransparency(bSupportsTransparency)
 			.OnUrlChanged(BIND_UOBJECT_DELEGATE(FOnTextChanged, HandleOnUrlChanged))
-			.OnBeforePopup(BIND_UOBJECT_DELEGATE(FOnBeforePopupDelegate, HandleOnBeforePopup));
+			.OnBeforePopup(BIND_UOBJECT_DELEGATE(FOnBeforePopupDelegate, HandleOnBeforePopup))
+			.OnConsoleMessage(BIND_UOBJECT_DELEGATE(FOnConsoleMessageDelegate, HandleOnConsoleMessage));
 
 		return WebBrowserWidget.ToSharedRef();
 	}
@@ -120,6 +121,11 @@ void UWebBrowser::HandleOnUrlChanged(const FText& InText)
 	OnUrlChanged.Broadcast(InText);
 }
 
+void UWebBrowser::HandleOnConsoleMessage(const FString& Message, const FString& Source, int32 Line, EWebBrowserConsoleLogSeverity Severity)
+{
+	OnConsoleMessage.Broadcast(Message, Source, Line);
+}
+
 bool UWebBrowser::HandleOnBeforePopup(FString URL, FString Frame)
 {
 	if (OnBeforePopup.IsBound())
diff --git a/Engine/Plugins/Runtime/WebBrowserWidget/Source/WebBrowserWidget/Public/WebBrowser.h b/Engine/Plugins/Runtime/WebBrowserWidget/Source/WebBrowserWidget/Public/WebBrowser.h
--- a/Engine/Plugins/Runtime/WebBrowserWidget/Source/WebBrowserWidget/Public/WebBrowser.h
+++ b/Engine/Plugins/Runtime/WebBrowserWidget/Source/WebBrowserWidget/Public/WebBrowser.h
@@ -6,6 +6,7 @@
 
 #include "WebBrowser.generated.h"
 
+enum class EWebBrowserConsoleLogSeverity;
 
 /**
  * 
@@ -18,6 +19,7 @@ class WEBBROWSERWIDGET_API UWebBrowser : public UWidget
 public:
 	DECLARE_DYNAMIC_MULTICAST_DELEGATE_OneParam(FOnUrlChanged, const FText&, Text);
 	DECLARE_DYNAMIC_MULTICAST_DELEGATE_TwoParams(FOnBeforePopup, FString, URL, FString, Frame);
+	DECLARE_DYNAMIC_MULTICAST_DELEGATE_ThreeParams(FOnConsoleMessage, const FString&, Message, const FString&, Source, int32, Line);
 
 	/**
 	 * Load the specified URL
@@ -66,6 +68,10 @@ public:
 	UPROPERTY(BlueprintAssignable, Category = "Web Browser|Event")
 	FOnBeforePopup OnBeforePopup;
 
+	/** Called when the browser has console spew to print */
+	UPROPERTY(BlueprintAssignable, Category = "Web Browser|Event")
+	FOnConsoleMessage OnConsoleMessage;
+
 public:
 
 	//~ Begin UWidget interface
@@ -95,6 +101,7 @@ protected:
 	virtual TSharedRef<SWidget> RebuildWidget() override;
 	// End of UWidget interface
 
+	void HandleOnConsoleMessage(const FString& Message, const FString& Source, int32 Line, EWebBrowserConsoleLogSeverity Severity);
 	void HandleOnUrlChanged(const FText& Text);
 	bool HandleOnBeforePopup(FString URL, FString Frame);
 };
diff --git a/Engine/Source/Runtime/WebBrowser/Private/CEF/CEFBrowserHandler.cpp b/Engine/Source/Runtime/WebBrowser/Private/CEF/CEFBrowserHandler.cpp
--- a/Engine/Source/Runtime/WebBrowser/Private/CEF/CEFBrowserHandler.cpp
+++ b/Engine/Source/Runtime/WebBrowser/Private/CEF/CEFBrowserHandler.cpp
@@ -492,8 +492,14 @@ CefResourceRequestHandler::ReturnValue FCEFBrowserHandler::OnBeforeResourceLoad(
 			}
 		}
 
-		Request->SetHeaderMap(HeaderMap);
-
+		if (Request->IsReadOnly())
+		{
+			LOG_CEF_LOAD("FCEFBrowserHandler::OnBeforeResourceLoad - readonly");
+		}
+		else
+		{
+			Request->SetHeaderMap(HeaderMap);
+		}
 		Callback->Continue(true);
 	}));
 
diff --git a/Engine/Source/Runtime/WebBrowser/Private/CEF/CEFBrowserHandler.h b/Engine/Source/Runtime/WebBrowser/Private/CEF/CEFBrowserHandler.h
--- a/Engine/Source/Runtime/WebBrowser/Private/CEF/CEFBrowserHandler.h
+++ b/Engine/Source/Runtime/WebBrowser/Private/CEF/CEFBrowserHandler.h
@@ -342,7 +342,7 @@ public:
 		return ResourceLoadCompleteDelegate;
 	}
 
-	DECLARE_DELEGATE_FiveParams(FOnConsoleMessageDelegate, CefRefPtr<CefBrowser> /*Browser*/, cef_log_severity_t /*level*/, const CefString& /*Message*/, const CefString& /*Source*/, int /*Line*/);
+	DECLARE_DELEGATE_FiveParams(FOnConsoleMessageDelegate, CefRefPtr<CefBrowser> /*Browser*/, cef_log_severity_t /*level*/, const CefString& /*Message*/, const CefString& /*Source*/, int32 /*Line*/);
 	FOnConsoleMessageDelegate& OnConsoleMessage()
 	{
 		return ConsoleMessageDelegate;
diff --git a/Engine/Source/Runtime/WebBrowser/Private/CEF/CEFWebBrowserWindow.cpp b/Engine/Source/Runtime/WebBrowser/Private/CEF/CEFWebBrowserWindow.cpp
--- a/Engine/Source/Runtime/WebBrowser/Private/CEF/CEFWebBrowserWindow.cpp
+++ b/Engine/Source/Runtime/WebBrowser/Private/CEF/CEFWebBrowserWindow.cpp
@@ -31,8 +31,9 @@
 #include "Windows/WindowsCursor.h"
 typedef FWindowsCursor FPlatformCursor;
 #elif PLATFORM_MAC
-#include "Mac/MacCursor.h"
 #include "Mac/CocoaThread.h"
+#include "Mac/MacApplication.h"
+#include "Mac/MacCursor.h"
 typedef FMacCursor FPlatformCursor;
 #else
 #endif
@@ -417,7 +418,7 @@ FCEFWebBrowserWindow::FCEFWebBrowserWindow(CefRefPtr<CefBrowser> InBrowser, CefR
 	, Ime(new FCEFImeHandler(InBrowser))
 #endif
 	, RHIRenderHelper(nullptr)
-#if PLATFORM_WINDOWS
+#if PLATFORM_WINDOWS || PLATFORM_MAC
 	, bInDirectHwndMode(false)
 #endif
 {
@@ -432,9 +433,8 @@ FCEFWebBrowserWindow::FCEFWebBrowserWindow(CefRefPtr<CefBrowser> InBrowser, CefR
 		ReleaseTextures();
 	}
 
-#if PLATFORM_WINDOWS
-	HWND NativeHandle = InternalCefBrowser->GetHost()->GetWindowHandle();
-	if (NativeHandle)
+#if PLATFORM_WINDOWS || PLATFORM_MAC
+	if (InternalCefBrowser->GetHost()->GetWindowHandle() != nullptr)
 	{
 		bInDirectHwndMode = true;
 	}
@@ -598,10 +598,12 @@ TOptional<FSlateRenderTransform> FCEFWebBrowserWindow::GetWebBrowserRenderTransf
 	return LocalRenderTransform;
 }
 
-bool FCEFWebBrowserWindow::IsInDirectHwndMode() const
+bool FCEFWebBrowserWindow::BlockInputInDirectHwndMode() const
 {
 #if PLATFORM_WINDOWS
 	return bInDirectHwndMode;
+#elif PLATFORM_MAC
+	return bInDirectHwndMode;
 #endif
 
 	return false;
@@ -622,6 +624,8 @@ void FCEFWebBrowserWindow::SetViewportSize(FIntPoint WindowSize, FIntPoint Windo
 		WindowDPIScaleFactor = ParentWindowPtr->GetNativeWindow()->GetDPIScaleFactor();
 	}
 
+	ViewportPos = WindowPos;
+
 	// Ignore sizes that can't be seen as it forces CEF to re-render whole image
 	if ((WindowSize.X > 0 && WindowSize.Y > 0 && ViewportSize != WindowSize) || WindowDPIScaleFactor != ViewportDPIScaleFactor)
 	{
@@ -647,6 +651,23 @@ void FCEFWebBrowserWindow::SetViewportSize(FIntPoint WindowSize, FIntPoint Windo
 
 				::SetWindowPos(NativeHandle, 0, WindowPos.X - ParentRect.left, WindowPos.Y - ParentRect.top, WindowSizeScaled.X, WindowSizeScaled.Y, 0);
 			}
+#elif PLATFORM_MAC
+			CefWindowHandle NativeWindowHandle = InternalCefBrowser->GetHost()->GetWindowHandle();
+			if (NativeWindowHandle)
+			{
+				NSView* browserView = CAST_CEF_WINDOW_HANDLE_TO_NSVIEW(NativeWindowHandle);
+				if (TSharedPtr<SWindow> ParentWindowPtr = ParentWindow.Pin())
+				{
+					NSWindow* parentWindow = (NSWindow*)ParentWindowPtr->GetNativeWindow()->GetOSWindowHandle();
+
+					const FVector2D CocoaPosition = FMacApplication::ConvertSlatePositionToCocoa(WindowPos.X, WindowPos.Y);
+					NSRect parentFrame = [parentWindow frame];
+					NSRect Rect = NSMakeRect(CocoaPosition.X - parentFrame.origin.x, (CocoaPosition.Y - parentFrame.origin.y) - WindowSize.Y, FMath::Max(WindowSize.X, 1), FMath::Max(WindowSize.Y, 1));
+					Rect = [parentWindow frameRectForContentRect : Rect];
+
+					[browserView setFrame : Rect] ;
+				}
+			}
 #endif
 
 			if (bFirstSize)
@@ -1131,16 +1152,33 @@ void FCEFWebBrowserWindow::PopulateCefKeyEvent(const FKeyEvent& InKeyEvent, CefK
 
 }
 
-bool FCEFWebBrowserWindow::OnKeyDown(const FKeyEvent& InKeyEvent)
+#if PLATFORM_MAC
+bool FilterSystemKeyChord(const FKeyEvent& InKeyEvent)
 {
-	if (IsValid() && !bIgnoreKeyDownEvent)
+	if(InKeyEvent.IsControlDown())
 	{
-#if PLATFORM_MAC
 		// Special case for Mac - make sure Cmd+~ is always passed back to the OS
-		if (InKeyEvent.GetKey() == EKeys::Tilde && InKeyEvent.IsControlDown())
+		if (InKeyEvent.GetKey() == EKeys::Tilde)
 		{
-			return false;
+			return true;
 		}
+		// Special case for Mac - make sure Cmd+H is always ignored by CEF
+		if (InKeyEvent.GetKey() == EKeys::H )
+		{
+			return true;
+		}
+	}
+	return false;
+}
+#endif
+
+bool FCEFWebBrowserWindow::OnKeyDown(const FKeyEvent& InKeyEvent)
+{
+	if (IsValid() && !BlockInputInDirectHwndMode() && !bIgnoreKeyDownEvent)
+	{
+#if PLATFORM_MAC
+		if(FilterSystemKeyChord(InKeyEvent))
+			return false;
 #endif
 		PreviousKeyDownEvent = InKeyEvent;
 		CefKeyEvent KeyEvent;
@@ -1154,14 +1192,11 @@ bool FCEFWebBrowserWindow::OnKeyDown(const FKeyEvent& InKeyEvent)
 
 bool FCEFWebBrowserWindow::OnKeyUp(const FKeyEvent& InKeyEvent)
 {
-	if (IsValid() && !bIgnoreKeyUpEvent)
+	if (IsValid() && !BlockInputInDirectHwndMode() && !bIgnoreKeyUpEvent)
 	{
 #if PLATFORM_MAC
-		// Special case for Mac - make sure Cmd+~ is always passed back to the OS
-		if (InKeyEvent.GetKey() == EKeys::Tilde && InKeyEvent.IsControlDown())
-		{
+		if(FilterSystemKeyChord(InKeyEvent))
 			return false;
-		}
 #endif
 		PreviousKeyUpEvent = InKeyEvent;
 		CefKeyEvent KeyEvent;
@@ -1175,7 +1210,7 @@ bool FCEFWebBrowserWindow::OnKeyUp(const FKeyEvent& InKeyEvent)
 
 bool FCEFWebBrowserWindow::OnKeyChar(const FCharacterEvent& InCharacterEvent)
 {
-	if (IsValid() && !bIgnoreCharacterEvent)
+	if (IsValid() && !BlockInputInDirectHwndMode() && !bIgnoreCharacterEvent)
 	{
 		PreviousCharacterEvent = InCharacterEvent;
 		CefKeyEvent KeyEvent;
@@ -1406,7 +1441,7 @@ void FCEFWebBrowserWindow::OnRenderProcessTerminated(CefRequestHandler::Terminat
 FReply FCEFWebBrowserWindow::OnMouseButtonDown(const FGeometry& MyGeometry, const FPointerEvent& MouseEvent, bool bIsPopup)
 {
 	FReply Reply = FReply::Unhandled();
-	if (IsValid() && !IsInDirectHwndMode())
+	if (IsValid() && !BlockInputInDirectHwndMode())
 	{
 		FKey Button = MouseEvent.GetEffectingButton();
 		// CEF only supports left, right, and middle mouse buttons
@@ -1436,7 +1471,7 @@ FReply FCEFWebBrowserWindow::OnMouseButtonDown(const FGeometry& MyGeometry, cons
 FReply FCEFWebBrowserWindow::OnMouseButtonUp(const FGeometry& MyGeometry, const FPointerEvent& MouseEvent, bool bIsPopup)
 {
 	FReply Reply = FReply::Unhandled();
-	if (IsValid() && !IsInDirectHwndMode())
+	if (IsValid() && !BlockInputInDirectHwndMode())
 	{
 		FKey Button = MouseEvent.GetEffectingButton();
 		// CEF only supports left, right, and middle mouse buttons
@@ -1482,7 +1517,7 @@ FReply FCEFWebBrowserWindow::OnMouseButtonUp(const FGeometry& MyGeometry, const
 FReply FCEFWebBrowserWindow::OnMouseButtonDoubleClick(const FGeometry& MyGeometry, const FPointerEvent& MouseEvent, bool bIsPopup)
 {
 	FReply Reply = FReply::Unhandled();
-	if (IsValid() && !IsInDirectHwndMode())
+	if (IsValid() && !BlockInputInDirectHwndMode())
 	{
 		FKey Button = MouseEvent.GetEffectingButton();
 		// CEF only supports left, right, and middle mouse buttons
@@ -1505,7 +1540,7 @@ FReply FCEFWebBrowserWindow::OnMouseButtonDoubleClick(const FGeometry& MyGeometr
 FReply FCEFWebBrowserWindow::OnMouseMove(const FGeometry& MyGeometry, const FPointerEvent& MouseEvent, bool bIsPopup)
 {
 	FReply Reply = FReply::Unhandled();
-	if (IsValid() && !IsInDirectHwndMode())
+	if (IsValid() && !BlockInputInDirectHwndMode())
 	{
 		CefMouseEvent Event = GetCefMouseEvent(MyGeometry, MouseEvent, bIsPopup);
 
@@ -1539,7 +1574,7 @@ void FCEFWebBrowserWindow::OnMouseLeave(const FPointerEvent& MouseEvent)
 	SetToolTip(CefString());
 	// We have no geometry here to convert our mouse event to local space so we just make a dummy event and set the moueLeave param to true
 	CefMouseEvent DummyEvent;
-	if (IsValid() && !IsInDirectHwndMode())
+	if (IsValid() && !BlockInputInDirectHwndMode())
 	{
 		InternalCefBrowser->GetHost()->SendMouseMoveEvent(DummyEvent, true);
 	}
@@ -1559,7 +1594,7 @@ bool FCEFWebBrowserWindow::GetSupportsMouseWheel() const
 FReply FCEFWebBrowserWindow::OnTouchGesture(const FGeometry& MyGeometry, const FPointerEvent& GestureEvent, bool bIsPopup)
 {
 	FReply Reply = FReply::Unhandled();
-	if(IsValid() && bSupportsMouseWheel)
+	if(IsValid() && bSupportsMouseWheel && !BlockInputInDirectHwndMode())
 	{
 		const EGestureEvent GestureType = GestureEvent.GetGestureType();
 		const FVector2D& GestureDelta = GestureEvent.GetGestureDelta();
@@ -1577,7 +1612,7 @@ FReply FCEFWebBrowserWindow::OnTouchGesture(const FGeometry& MyGeometry, const F
 FReply FCEFWebBrowserWindow::OnMouseWheel(const FGeometry& MyGeometry, const FPointerEvent& MouseEvent, bool bIsPopup)
 {
 	FReply Reply = FReply::Unhandled();
-	if(IsValid() && bSupportsMouseWheel && !IsInDirectHwndMode())
+	if(IsValid() && bSupportsMouseWheel && !BlockInputInDirectHwndMode())
 	{
 #if PLATFORM_WINDOWS
 		// The original delta is reduced so this should bring it back to what CEF expects
@@ -1705,6 +1740,15 @@ void FCEFWebBrowserWindow::CloseBrowser(bool bForce, bool bBlockTillClosed)
 	if (IsValid())
 	{
 		CefRefPtr<CefBrowserHost> Host = InternalCefBrowser->GetHost();
+#if PLATFORM_MAC
+		CefWindowHandle NativeWindowHandle = Host->GetWindowHandle();
+		if (NativeWindowHandle != nullptr)
+		{
+			NSView *browserView = CAST_CEF_WINDOW_HANDLE_TO_NSVIEW(NativeWindowHandle);
+			if (browserView != nil)
+				[browserView removeFromSuperview];
+		}
+#endif
 		// In case this is called from inside a CEF event handler, use CEF's task mechanism to
 		// postpone the actual closing of the window until it is safe.
 		CefPostTask(TID_UI, new FCEFBrowserClosureTask(nullptr, [=]()
@@ -2217,6 +2261,36 @@ bool FCEFWebBrowserWindow::OnCursorChange(CefCursorHandle CefCursor, cef_cursor_
 	return false;
 }
 
+EWebTransitionSource TransitionTypeToSourceEnum(const CefRequest::TransitionType& Type)
+{
+	CefRequest::TransitionType TransitionSource = (CefRequest::TransitionType)(Type & TT_SOURCE_MASK);
+	switch (TransitionSource)
+	{
+		case TT_LINK:			return EWebTransitionSource::Link;
+		case TT_EXPLICIT:		return EWebTransitionSource::Explicit;
+		case TT_AUTO_SUBFRAME:	return EWebTransitionSource::AutoSubframe;
+		case TT_FORM_SUBMIT:	return EWebTransitionSource::FormSubmit;
+		case TT_RELOAD:			return EWebTransitionSource::Reload;
+		default:				return EWebTransitionSource::Unknown;
+	}
+	return EWebTransitionSource::Unknown;
+}
+
+EWebTransitionSourceQualifier TransitionTypeToSourceQualifierEnum(const CefRequest::TransitionType& Type)
+{
+	CefRequest::TransitionType TransitionSourceQualifier = (CefRequest::TransitionType)(Type & TT_QUALIFIER_MASK);
+	switch (TransitionSourceQualifier)
+	{
+		case TT_BLOCKED_FLAG:			return EWebTransitionSourceQualifier::Blocked;
+		case TT_FORWARD_BACK_FLAG:		return EWebTransitionSourceQualifier::ForwardBack;
+		case TT_CHAIN_START_FLAG:		return EWebTransitionSourceQualifier::ChainStart;
+		case TT_CHAIN_END_FLAG:			return EWebTransitionSourceQualifier::ChainEnd;
+		case TT_CLIENT_REDIRECT_FLAG:	return EWebTransitionSourceQualifier::ClientRedirect;
+		case TT_SERVER_REDIRECT_FLAG:	return EWebTransitionSourceQualifier::ServerRedirect;
+		default:						return EWebTransitionSourceQualifier::Unknown;
+	}
+	return EWebTransitionSourceQualifier::Unknown;
+}
 
 bool FCEFWebBrowserWindow::OnBeforeBrowse( CefRefPtr<CefBrowser> Browser, CefRefPtr<CefFrame> Frame, CefRefPtr<CefRequest> Request, bool user_gesture, bool bIsRedirect )
 {
@@ -2233,8 +2307,12 @@ bool FCEFWebBrowserWindow::OnBeforeBrowse( CefRefPtr<CefBrowser> Browser, CefRef
 				FWebNavigationRequest RequestDetails;
 				RequestDetails.bIsRedirect = bIsRedirect;
 				RequestDetails.bIsMainFrame = bIsMainFrame;
-				RequestDetails.bIsExplicitTransition = Request->GetTransitionType() == TT_EXPLICIT;
 
+				const CefRequest::TransitionType RequestTransitionType = Request->GetTransitionType();
+				RequestDetails.TransitionSource = TransitionTypeToSourceEnum(RequestTransitionType);
+				RequestDetails.TransitionSourceQualifier = TransitionTypeToSourceQualifierEnum(RequestTransitionType);
+				RequestDetails.bIsExplicitTransition = RequestDetails.TransitionSource == EWebTransitionSource::Explicit;
+			
 				if (bIsMainFrame)
 				{
 					// We need to defer all future navigations until we can determine if this current navigation is going to be handled or not
@@ -2313,8 +2391,8 @@ EWebBrowserConsoleLogSeverity CefLogSeverityToWebBrowser(cef_log_severity_t Leve
 	{
 	case LOGSEVERITY_VERBOSE:
 		return EWebBrowserConsoleLogSeverity::Verbose;
-		//case LOGSEVERITY_DEBUG: // same as LOGSEVERITY_VERBOSE
-		//	return Verbose;
+	//case LOGSEVERITY_DEBUG: // same as LOGSEVERITY_VERBOSE
+	//	return Verbose;
 	case LOGSEVERITY_INFO:
 		return EWebBrowserConsoleLogSeverity::Info;
 	case LOGSEVERITY_WARNING:
@@ -2329,7 +2407,7 @@ EWebBrowserConsoleLogSeverity CefLogSeverityToWebBrowser(cef_log_severity_t Leve
 	}
 }
 
-void FCEFWebBrowserWindow::HandleOnConsoleMessage(CefRefPtr<CefBrowser> Browser, cef_log_severity_t Level, const CefString& Message, const CefString& Source, int Line)
+void FCEFWebBrowserWindow::HandleOnConsoleMessage(CefRefPtr<CefBrowser> Browser, cef_log_severity_t Level, const CefString& Message, const CefString& Source, int32 Line)
 {
 	ConsoleMessageDelegate.ExecuteIfBound(WCHAR_TO_TCHAR(Message.ToWString().c_str()), WCHAR_TO_TCHAR(Source.ToWString().c_str()), Line, CefLogSeverityToWebBrowser(Level));
 }
@@ -2416,6 +2494,7 @@ int32 FCEFWebBrowserWindow::GetCefMouseModifiers(const FPointerEvent& InMouseEve
 CefMouseEvent FCEFWebBrowserWindow::GetCefMouseEvent(const FGeometry& MyGeometry, const FPointerEvent& MouseEvent, bool bIsPopup)
 {
 	CefMouseEvent Event;
+
 	FGeometry MouseGeometry = MyGeometry;
 	if (bUsingAcceleratedPaint)
 	{
@@ -2498,12 +2577,12 @@ bool FCEFWebBrowserWindow::CanSupportAcceleratedPaint()
 	return false;
 #elif PLATFORM_WINDOWS
 #if PLATFORM_64BITS
-	//return false;
-	
+	return false;
+	/*
 	static bool Windows10OrAbove = FWindowsPlatformMisc::VerifyWindowsVersion(10, 0); //Win10
 	if (Windows10OrAbove == false)
 	{
-		return false;
+		return  -false;
 	}
 
 	// match the logic in GetStandardStandaloneRenderer() from StandaloneRenderer.cpp to check for the OGL slate renderer
@@ -2511,8 +2590,8 @@ bool FCEFWebBrowserWindow::CanSupportAcceleratedPaint()
 	{
 		return false;
 	}
-	// Disable accelerated paint by default until we can fix the multi-gpu rendering issue
-	return false;
+	return true;*/
+
 #else
 	return false; // 32-bit windows doesn't have the accelerated rendering patches applied, it can be done if needed
 #endif
@@ -2599,9 +2678,16 @@ void FCEFWebBrowserWindow::ProcessPendingNavigation()
 		CefString Url = TCHAR_TO_WCHAR(*PendingLoadUrl);
 		PendingLoadUrl.Empty();
 #if PLATFORM_MAC
-		MainThreadCall(^{
-		MainFrame->LoadURL(Url);
-		}, NSDefaultRunLoopMode, true);
+		if ([NSThread isMainThread])
+		{
+			MainFrame->LoadURL(Url);
+		}
+		else
+		{
+			MainThreadCall(^{
+				MainFrame->LoadURL(Url);
+			}, NSDefaultRunLoopMode, true);
+		}
 #else
 		MainFrame->LoadURL(Url);
 #endif
@@ -2632,14 +2718,49 @@ void FCEFWebBrowserWindow::SetIsHidden(bool bValue)
 				{
 					::SetFocus((HWND)ParentWindowPtr->GetNativeWindow()->GetOSWindowHandle());
 				}
+				// when hidden also resize the window to 0x0 to further reduce resource usage. This copies the
+				// behavior of the CefClient code, see browser_window_std_win.cc in the ::Hide() function. This code
+				// is also required for the HTML5 visibility API to work 
+				SetWindowPos(NativeWindowHandle, NULL, 0, 0, 0, 0, SWP_NOZORDER | SWP_NOMOVE | SWP_NOACTIVATE);
 			}
+			else
+			{
+				// restore the window to its right size/position
+				HWND Parent = ::GetParent(NativeWindowHandle);
+				// Position is in screen coordinates, so we'll need to get the parent window location first.
+				RECT ParentRect = { 0, 0, 0, 0 };
+				if (Parent)
+				{
+					::GetWindowRect(Parent, &ParentRect);
+				}
+
+				FIntPoint WindowSizeScaled = (FVector2D(ViewportSize) * ViewportDPIScaleFactor).IntPoint();
+
+				::SetWindowPos(NativeWindowHandle, 0, ViewportPos.X - ParentRect.left, ViewportPos.Y - ParentRect.top, WindowSizeScaled.X, WindowSizeScaled.Y, 0);
+			}
+		}
+		else
+		{
+#elif PLATFORM_MAC
+		CefWindowHandle NativeWindowHandle = BrowserHost->GetWindowHandle();
+		if (NativeWindowHandle != nullptr)
+		{
+			NSView *browserView = CAST_CEF_WINDOW_HANDLE_TO_NSVIEW(NativeWindowHandle);
+			if(bIsHidden)
+			{
+				[browserView setHidden:YES];
+			}
+			else
+			{
+				[browserView setHidden:NO];
+			}
+			
 		}
 		else
 		{
-			// only hide OSR renderer windows
 #endif
 			BrowserHost->WasHidden(bIsHidden);
-#if PLATFORM_WINDOWS
+#if PLATFORM_WINDOWS || PLATFORM_MAC
 		}
 #endif
 	}
@@ -2702,7 +2823,12 @@ bool FCEFWebBrowserWindow::OnProcessMessageReceived(CefRefPtr<CefBrowser> Browse
 
 	if (!bHandled)
 	{
-#if !PLATFORM_LINUX
+#if PLATFORM_MAC
+		if (!bInDirectHwndMode) // IME is handled by the CEF control in direct render mode
+		{
+			bHandled = Ime->OnProcessMessageReceived(Browser, SourceProcess, Message);
+		}
+#elif PLATFORM_WINDOWS
 		bHandled = Ime->OnProcessMessageReceived(Browser, SourceProcess, Message);
 #endif
 	}
diff --git a/Engine/Source/Runtime/WebBrowser/Private/CEF/CEFWebBrowserWindow.h b/Engine/Source/Runtime/WebBrowser/Private/CEF/CEFWebBrowserWindow.h
--- a/Engine/Source/Runtime/WebBrowser/Private/CEF/CEFWebBrowserWindow.h
+++ b/Engine/Source/Runtime/WebBrowser/Private/CEF/CEFWebBrowserWindow.h
@@ -399,7 +399,7 @@ private:
 
 	void HandleOnBeforeResourceLoad(const CefString& URL, CefRequest::ResourceType Type, FRequestHeaders& AdditionalHeaders, const bool AllowUserCredentials);
 	void HandleOnResourceLoadComplete(const CefString& URL, CefRequest::ResourceType Type, CefResourceRequestHandler::URLRequestStatus Status, int64 ContentLength);
-	void HandleOnConsoleMessage(CefRefPtr<CefBrowser> Browser, cef_log_severity_t Level, const CefString& Message, const CefString& Source, int Line);
+	void HandleOnConsoleMessage(CefRefPtr<CefBrowser> Browser, cef_log_severity_t Level, const CefString& Message, const CefString& Source, int32 Line);
 
 	/**
 	 * Called before loading a resource to allow overriding the content for a request.
@@ -589,7 +589,7 @@ private:
 	/** Used to let us correctly render the web texture for the accelerated render path */
 	TOptional<FSlateRenderTransform> GetWebBrowserRenderTransform() const;
 
-	bool IsInDirectHwndMode() const;
+	bool BlockInputInDirectHwndMode() const;
 
 #if PLATFORM_WINDOWS
 	/** manually load cursor icons from the CEF3 dll if needed, working around a CEF3 bug */
@@ -621,6 +621,9 @@ private:
 	/** Current size of this window. */
 	FIntPoint ViewportSize;
 
+	/** Current position of this window. */
+	FIntPoint ViewportPos;
+
 	/** Current DPI scale factor of this window. */
 	float ViewportDPIScaleFactor;
 
@@ -772,7 +775,7 @@ private:
 	TWeakPtr<SWindow> ParentWindow;
 	FCEFWebBrowserWindowRHIHelper* RHIRenderHelper;
 	
-#if PLATFORM_WINDOWS
+#if PLATFORM_WINDOWS || PLATFORM_MAC
 	bool bInDirectHwndMode;
 #endif
 };
diff --git a/Engine/Source/Runtime/WebBrowser/Private/SWebBrowser.cpp b/Engine/Source/Runtime/WebBrowser/Private/SWebBrowser.cpp
--- a/Engine/Source/Runtime/WebBrowser/Private/SWebBrowser.cpp
+++ b/Engine/Source/Runtime/WebBrowser/Private/SWebBrowser.cpp
@@ -28,6 +28,7 @@ void SWebBrowser::Construct(const FArguments& InArgs, const TSharedPtr<IWebBrows
 	OnShowDialog = InArgs._OnShowDialog;
 	OnDismissAllDialogs = InArgs._OnDismissAllDialogs;
 	OnBeforePopup = InArgs._OnBeforePopup;
+	OnConsoleMessage = InArgs._OnConsoleMessage;
 	OnCreateWindow = InArgs._OnCreateWindow;
 	OnCloseWindow = InArgs._OnCloseWindow;
 	bShowInitialThrobber = InArgs._ShowInitialThrobber;
@@ -126,6 +127,7 @@ void SWebBrowser::Construct(const FArguments& InArgs, const TSharedPtr<IWebBrows
 				.Visibility(this, &SWebBrowser::GetViewportVisibility)
 				.OnSuppressContextMenu(InArgs._OnSuppressContextMenu)
 				.OnDragWindow(InArgs._OnDragWindow)
+				.OnConsoleMessage(OnConsoleMessage)
 				.BrowserFrameRate(InArgs._BrowserFrameRate)
 			]
 			+ SOverlay::Slot()
diff --git a/Engine/Source/Runtime/WebBrowser/Private/SWebBrowserView.cpp b/Engine/Source/Runtime/WebBrowser/Private/SWebBrowserView.cpp
--- a/Engine/Source/Runtime/WebBrowser/Private/SWebBrowserView.cpp
+++ b/Engine/Source/Runtime/WebBrowser/Private/SWebBrowserView.cpp
@@ -47,6 +47,9 @@ SWebBrowserView::~SWebBrowserView()
 		BrowserWindow->OnDismissAllDialogs().Unbind();
 		BrowserWindow->OnCreateWindow().Unbind();
 		BrowserWindow->OnCloseWindow().Unbind();
+		BrowserWindow->OnSuppressContextMenu().Unbind();
+		BrowserWindow->OnDragWindow().Unbind();
+		BrowserWindow->OnConsoleMessage().Unbind();
 
 		if (BrowserWindow->OnBeforeBrowse().IsBoundToObject(this))
 		{
@@ -95,6 +98,7 @@ void SWebBrowserView::Construct(const FArguments& InArgs, const TSharedPtr<IWebB
 	OnUnhandledKeyDown = InArgs._OnUnhandledKeyDown;
 	OnUnhandledKeyUp = InArgs._OnUnhandledKeyUp;
 	OnUnhandledKeyChar = InArgs._OnUnhandledKeyChar;
+	OnConsoleMessage = InArgs._OnConsoleMessage;
 
 	BrowserWindow = InWebBrowserWindow;
 	if(!BrowserWindow.IsValid())
@@ -211,6 +215,11 @@ void SWebBrowserView::Construct(const FArguments& InArgs, const TSharedPtr<IWebB
 		BrowserWindow->OnDragWindow().BindSP(this, &SWebBrowserView::HandleDrag);
 		OnDragWindow = InArgs._OnDragWindow;
 
+		if (!BrowserWindow->OnConsoleMessage().IsBound())
+		{
+			BrowserWindow->OnConsoleMessage().BindSP(this, &SWebBrowserView::HandleConsoleMessage);
+		}
+
 		BrowserViewport = MakeShareable(new FWebBrowserViewport(BrowserWindow));
 #if WITH_CEF3
 		BrowserWidget->SetViewportInterface(BrowserViewport.ToSharedRef());
@@ -466,6 +475,7 @@ void SWebBrowserView::CloseBrowser()
 {
 	BrowserWindow->CloseBrowser(true /*force*/, true /*block until closed*/);
 }
+
 void SWebBrowserView::HandleToolTip(FString ToolTipText)
 {
 	if(ToolTipText.IsEmpty())
@@ -728,7 +738,6 @@ bool SWebBrowserView::UnhandledKeyChar(const FCharacterEvent& CharacterEvent)
 	return false;
 }
 
-
 void SWebBrowserView::SetParentWindow(TSharedPtr<SWindow> Window)
 {
 	SetupParentWindowHandlers();
@@ -743,4 +752,9 @@ void SWebBrowserView::SetBrowserKeyboardFocus()
 	BrowserWindow->OnFocus(HasAnyUserFocusOrFocusedDescendants(), false);
 }
 
+void SWebBrowserView::HandleConsoleMessage(const FString& Message, const FString& Source, int32 Line, EWebBrowserConsoleLogSeverity Serverity)
+{
+	OnConsoleMessage.ExecuteIfBound(Message, Source, Line, Serverity);
+}
+
 #undef LOCTEXT_NAMESPACE
diff --git a/Engine/Source/Runtime/WebBrowser/Public/IWebBrowserWindow.h b/Engine/Source/Runtime/WebBrowser/Public/IWebBrowserWindow.h
--- a/Engine/Source/Runtime/WebBrowser/Public/IWebBrowserWindow.h
+++ b/Engine/Source/Runtime/WebBrowser/Public/IWebBrowserWindow.h
@@ -33,11 +33,49 @@ enum class EWebBrowserConsoleLogSeverity
 	Fatal
 };
 
+
+enum class EWebTransitionSource
+{
+	Unknown,
+	/** Source is a link click or the JavaScript window.open function. */
+	Link,
+	/** Source is some other "explicit" navigation action such as creating a new browser or using the LoadURL function. */
+	Explicit,
+	/** Source is a subframe navigation. */
+	AutoSubframe,
+	/** Source is a subframe navigation explicitly requested by the user. */
+	ManualSubframe,
+	/** Source is a form submission by the user. */
+	FormSubmit,
+	/** Source is a "reload" of the page via the Reload function */
+	Reload
+};
+
+enum class EWebTransitionSourceQualifier
+{
+	Unknown,
+	/** Attempted to visit a URL but was blocked. */
+	Blocked,
+	/** Used the Forward or Back function to navigate among browsing history. */
+	ForwardBack,
+	/** The beginning of a navigation chain. */
+	ChainStart,
+	/** The last transition in a redirect chain. */
+	ChainEnd,
+	/** Redirects caused by JavaScript or a meta refresh tag on the page. */
+	ClientRedirect,
+	/** Used to test whether a transition involves a redirect. */
+	ServerRedirect
+};
+
+
 struct FWebNavigationRequest
 {
 	bool bIsRedirect;
 	bool bIsMainFrame;
 	bool bIsExplicitTransition;
+	EWebTransitionSource TransitionSource;
+	EWebTransitionSourceQualifier TransitionSourceQualifier;
 };
 
 /**
@@ -375,7 +413,7 @@ public:
 	virtual FOnResourceLoadCompleteDelegate& OnResourceLoadComplete() = 0;
 
 	/** A delegate that is invoked for each console message */
-	DECLARE_DELEGATE_FourParams(FOnConsoleMessageDelegate, const FString& /*Message*/, const FString& /*Source*/, int /*Line*/, EWebBrowserConsoleLogSeverity /*severity*/);
+	DECLARE_DELEGATE_FourParams(FOnConsoleMessageDelegate, const FString& /*Message*/, const FString& /*Source*/, int32 /*Line*/, EWebBrowserConsoleLogSeverity /*severity*/);
 	virtual FOnConsoleMessageDelegate& OnConsoleMessage() = 0;
 
 	/** A delegate that is invoked when an existing browser requests creation of a new browser window. */
diff --git a/Engine/Source/Runtime/WebBrowser/Public/SWebBrowser.h b/Engine/Source/Runtime/WebBrowser/Public/SWebBrowser.h
--- a/Engine/Source/Runtime/WebBrowser/Public/SWebBrowser.h
+++ b/Engine/Source/Runtime/WebBrowser/Public/SWebBrowser.h
@@ -121,6 +121,8 @@ public:
 		/** Called when drag is detected in a web page area tagged as a drag region. */
 		SLATE_EVENT(FOnDragWindow, OnDragWindow);
 
+		/** Called for each console message */
+		SLATE_EVENT(FOnConsoleMessageDelegate, OnConsoleMessage);
 
 	SLATE_END_ARGS()
 
@@ -200,7 +202,7 @@ public:
 	 * Properties and Functions will be accessible from JavaScript side.
 	 * As all communication with the rendering procesis asynchronous, return values (both for properties and function results) are wrapped into JS Future objects.
 	 *
-	 * @param Name The name of the object. The object will show up as window.ue4.{Name} on the javascript side. Note: All object names, function names, and property names will be converted to lower case when bound on the javascript side.  If there is an existing object of the same name, this object will replace it. If bIsPermanent is false and there is an existing permanent binding, the permanent binding will be restored when the temporary one is removed.
+	 * @param Name The name of the object. The object will show up as window.ue.{Name} on the javascript side. Note: All object names, function names, and property names will be converted to lower case when bound on the javascript side.  If there is an existing object of the same name, this object will replace it. If bIsPermanent is false and there is an existing permanent binding, the permanent binding will be restored when the temporary one is removed.
 	 * @param Object The object instance.
 	 * @param bIsPermanent If true, the object will be visible to all pages loaded through this browser widget, otherwise, it will be deleted when navigating away from the current page. Non-permanent bindings should be registered from inside an OnLoadStarted event handler in order to be available before JS code starts loading.
 	 */
@@ -287,6 +289,9 @@ private:
 	/** A delegate that is invoked when the browser attempts to pop up a new window */
 	FOnBeforePopupDelegate OnBeforePopup;
 
+	/** A delegate that is invoked when the browser writes console output */
+	FOnConsoleMessageDelegate OnConsoleMessage;
+
 	/** A delegate that is invoked when the browser requests a UI window for another browser it spawned */
 	FOnCreateWindowDelegate OnCreateWindow;
 
diff --git a/Engine/Source/Runtime/WebBrowser/Public/SWebBrowserView.h b/Engine/Source/Runtime/WebBrowser/Public/SWebBrowserView.h
--- a/Engine/Source/Runtime/WebBrowser/Public/SWebBrowserView.h
+++ b/Engine/Source/Runtime/WebBrowser/Public/SWebBrowserView.h
@@ -25,6 +25,7 @@ DECLARE_DELEGATE_RetVal_TwoParams(bool, FOnBeforePopupDelegate, FString, FString
 DECLARE_DELEGATE_RetVal_TwoParams(bool, FOnCreateWindowDelegate, const TWeakPtr<IWebBrowserWindow>&, const TWeakPtr<IWebBrowserPopupFeatures>&);
 DECLARE_DELEGATE_RetVal_OneParam(bool, FOnCloseWindowDelegate, const TWeakPtr<IWebBrowserWindow>&);
 DECLARE_DELEGATE_RetVal_OneParam(TSharedPtr<IToolTip>, FOnCreateToolTip, const FText&);
+DECLARE_DELEGATE_FourParams(FOnConsoleMessageDelegate, const FString& /*Message*/, const FString& /*Source*/, int32 /*Line*/, EWebBrowserConsoleLogSeverity /*Severity*/);
 
 #if WITH_CEF3
 typedef SViewport SWebBrowserWidget;
@@ -45,6 +46,7 @@ public:
 	DECLARE_DELEGATE_RetVal_OneParam(bool, FOnUnhandledKeyUp, const FKeyEvent& /*KeyEvent*/);
 	DECLARE_DELEGATE_RetVal_OneParam(bool, FOnUnhandledKeyChar, const FCharacterEvent& /*CharacterEvent*/);
 
+
 	SLATE_BEGIN_ARGS(SWebBrowserView)
 		: _InitialURL(TEXT("https://www.google.com"))
 		, _ShowErrorMessage(true)
@@ -151,6 +153,9 @@ public:
 
 		/** Called to allow the handling of any key char events not handled by the browser. */
 		SLATE_EVENT(FOnUnhandledKeyChar, OnUnhandledKeyChar)
+		
+		/** Called for each console message */
+		SLATE_EVENT(FOnConsoleMessageDelegate, OnConsoleMessage)
 
 	SLATE_END_ARGS()
 
@@ -356,6 +361,7 @@ private:
 	bool UnhandledKeyChar(const FCharacterEvent& CharacterEvent);
 
 	bool HandleDrag(const FPointerEvent& MouseEvent);
+	void HandleConsoleMessage(const FString& Message, const FString& Source, int32 Line, EWebBrowserConsoleLogSeverity Serverity);
 
 	TOptional<FSlateRenderTransform> GetPopupRenderTransform() const;
 private:
@@ -429,7 +435,7 @@ private:
 	/** A delegate that is invoked when the browser detects drag event in within drag region */
 	FOnDragWindow OnDragWindow;
 	
-		/** A delegate for handling key down events not handled by browser. */
+	/** A delegate for handling key down events not handled by browser. */
 	FOnUnhandledKeyDown OnUnhandledKeyDown;
 
 	/** A delegate for handling key up events not handled by browser. */
@@ -437,6 +443,9 @@ private:
 
 	/** A delegate for handling key char events not handled by browser. */
 	FOnUnhandledKeyChar OnUnhandledKeyChar;
+	
+	/** A delegate that is invoked for each console message */
+	FOnConsoleMessageDelegate OnConsoleMessage;
 
 protected:
 	bool HandleSuppressContextMenu();
