- Don't attempt to backup the cef3.log file from CEF when running in commandlet mode.

diff --git a/Engine/Source/Runtime/WebBrowser/Private/WebBrowserModule.cpp b/Engine/Source/Runtime/WebBrowser/Private/WebBrowserModule.cpp
--- a/Engine/Source/Runtime/WebBrowser/Private/WebBrowserModule.cpp
+++ b/Engine/Source/Runtime/WebBrowser/Private/WebBrowserModule.cpp
@@ -56,7 +56,10 @@ IMPLEMENT_MODULE( FWebBrowserModule, WebBrowser );
 void FWebBrowserModule::StartupModule()
 {
 #if WITH_CEF3
-	CEF3Utils::BackupCEF3Logfile(FPaths::ProjectLogDir());
+	if (!IsRunningCommandlet())
+	{
+		CEF3Utils::BackupCEF3Logfile(FPaths::ProjectLogDir());
+	}
 	bLoadedCEFModule = CEF3Utils::LoadCEF3Modules(true);
 #if PLATFORM_MAC
 	// Dynamically load the CEF framework library into this dylibs memory space.
