Add a dealloc to IOSWebViewWrapper (via Static Analysis) and a few other retained properties found with latest clang.

#rnx

diff --git a/Engine/Source/Runtime/WebBrowser/Private/IOS/IOSPlatformWebBrowser.cpp b/Engine/Source/Runtime/WebBrowser/Private/IOS/IOSPlatformWebBrowser.cpp
--- a/Engine/Source/Runtime/WebBrowser/Private/IOS/IOSPlatformWebBrowser.cpp
+++ b/Engine/Source/Runtime/WebBrowser/Private/IOS/IOSPlatformWebBrowser.cpp
@@ -485,6 +485,7 @@ supportsMetal : (bool)InSupportsMetal supportsMetalMRT : (bool)InSupportsMetalMR
 			[self.WebView setOpaque : YES];
 		}
 
+		[theConfiguration release];
 		[self setDefaultVisibility];
 	});
 #endif
@@ -498,12 +499,37 @@ supportsMetal : (bool)InSupportsMetal supportsMetalMRT : (bool)InSupportsMetalMR
 	{
 		[self.WebViewContainer removeFromSuperview];
 		[self.WebView removeFromSuperview];
+		
+		[WebView release];
+		[WebViewContainer release];
+
 		WebView = nil;
 		WebViewContainer = nil;
 	});
 #endif
 }
 
+-(void)dealloc;
+{
+#if !PLATFORM_TVOS
+	if (WebView != nil)
+	{
+		WebView.navigationDelegate = nil;
+		[WebView release];
+		WebView = nil;
+	};
+
+	[WebViewContainer release];
+	WebViewContainer = nil;
+#endif
+	[NextContent release];
+	NextContent = nil;
+	[NextURL release];
+	NextURL = nil;
+
+	[super dealloc];
+}
+
 -(void)updateframe:(CGRect)InFrame;
 {
 	self.DesiredFrame = InFrame;
@@ -590,7 +616,7 @@ supportsMetal : (bool)InSupportsMetal supportsMetalMRT : (bool)InSupportsMetalMR
 	dispatch_async(dispatch_get_main_queue(), ^
 	{
 		self.NextContent = InString;
-	self.NextURL = InURL;
+		self.NextURL = InURL;
 	});
 }
 
diff --git a/Engine/Source/Runtime/WebBrowser/Private/IOS/IOSPlatformWebBrowser.h b/Engine/Source/Runtime/WebBrowser/Private/IOS/IOSPlatformWebBrowser.h
--- a/Engine/Source/Runtime/WebBrowser/Private/IOS/IOSPlatformWebBrowser.h
+++ b/Engine/Source/Runtime/WebBrowser/Private/IOS/IOSPlatformWebBrowser.h
@@ -51,6 +51,7 @@ class SWebBrowserView;
 -(void)create:(TSharedPtr<SIOSWebBrowserWidget>)InWebBrowserWidget useTransparency : (bool)InUseTransparency
 supportsMetal : (bool)InSupportsMetal supportsMetalMRT : (bool)InSupportsMetalMRT;
 -(void)close;
+-(void)dealloc;
 -(void)updateframe:(CGRect)InFrame;
 -(void)loadstring:(NSString*)InString dummyurl:(NSURL*)InURL;
 -(void)loadurl:(NSURL*)InURL;
@@ -69,10 +70,10 @@ supportsMetal : (bool)InSupportsMetal supportsMetalMRT : (bool)InSupportsMetalMR
 -(void)SetVideoTextureValid:(bool)Condition;
 -(bool)IsVideoTextureValid;
 -(bool)UpdateVideoFrame:(void*)ptr;
-- (void)updateWebViewMetalTexture : (id<MTLTexture>)texture;
+-(void)updateWebViewMetalTexture : (id<MTLTexture>)texture;
 #if !PLATFORM_TVOS
-- (void)webView:(WKWebView*)InWebView decidePolicyForNavigationAction : (WKNavigationAction*)InNavigationAction decisionHandler : (void(^)(WKNavigationActionPolicy))InDecisionHandler;
--(void)webView:(WKWebView *)InWebView didCommitNavigation : (WKNavigation *)InNavigation;
+-(void)webView:(WKWebView*)InWebView decidePolicyForNavigationAction : (WKNavigationAction*)InNavigationAction decisionHandler : (void(^)(WKNavigationActionPolicy))InDecisionHandler;
+-(void)webView:(WKWebView*)InWebView didCommitNavigation : (WKNavigation*)InNavigation;
 #endif
 @end
 
