Fix for an invalid TWeakPtr dereference on ProcessScriptMessage() AsyncTask callback after iOS webview is closed.
Also, fixes a deref of a (possibly) destroyed WebBrowserWidget within a AsyncTask callback when notifying on Url Loading when Widget is closed.

diff --git a/Engine/Source/Runtime/WebBrowser/Private/IOS/IOSPlatformWebBrowser.cpp b/Engine/Source/Runtime/WebBrowser/Private/IOS/IOSPlatformWebBrowser.cpp
--- a/Engine/Source/Runtime/WebBrowser/Private/IOS/IOSPlatformWebBrowser.cpp
+++ b/Engine/Source/Runtime/WebBrowser/Private/IOS/IOSPlatformWebBrowser.cpp
@@ -306,24 +306,35 @@ class SIOSWebBrowserWidget : public SLeafWidget
 
 	bool HandleShouldOverrideUrlLoading(const FString& Url)
 	{
-		bool Retval = false;
 		if (WebBrowserWindowPtr.IsValid())
 		{
-			TSharedPtr<FWebBrowserWindow> BrowserWindow = WebBrowserWindowPtr.Pin();
-			if (BrowserWindow.IsValid())
+			// Capture vars needed for AsyncTask
+			FString UrlString = Url;
+			TWeakPtr<FWebBrowserWindow> AsyncWebBrowserWindowPtr = WebBrowserWindowPtr;
+
+			// Notify on the game thread
+			[FIOSAsyncTask CreateTaskWithBlock : ^ bool(void)
 			{
-				if (BrowserWindow->OnBeforeBrowse().IsBound())
+				if (AsyncWebBrowserWindowPtr.IsValid())
 				{
-					FWebNavigationRequest RequestDetails;
-					RequestDetails.bIsRedirect = false;
-					RequestDetails.bIsMainFrame = true; // shouldOverrideUrlLoading is only called on the main frame
-
-					Retval = BrowserWindow->OnBeforeBrowse().Execute(Url, RequestDetails);
-					BrowserWindow->SetTitle("");
+					TSharedPtr<FWebBrowserWindow> BrowserWindow = AsyncWebBrowserWindowPtr.Pin();
+					if (BrowserWindow.IsValid())
+					{
+						if (BrowserWindow->OnBeforeBrowse().IsBound())
+						{
+							FWebNavigationRequest RequestDetails;
+							RequestDetails.bIsRedirect = false;
+							RequestDetails.bIsMainFrame = true; // shouldOverrideUrlLoading is only called on the main frame
+							
+							BrowserWindow->OnBeforeBrowse().Execute(Url, RequestDetails);
+							BrowserWindow->SetTitle("");
+						}
+					}
 				}
-			}
+				return true;
+			}];
 		}
-		return Retval;
+		return true;
 	}
 
 	void HandleReceivedTitle(const FString& Title)
@@ -340,30 +351,35 @@ class SIOSWebBrowserWidget : public SLeafWidget
 
 	void ProcessScriptMessage(const FString& InMessage)
 	{
-		FString Message = InMessage;
 		if (WebBrowserWindowPtr.IsValid())
 		{
+			FString Message = InMessage;
+			TWeakPtr<FWebBrowserWindow> AsyncWebBrowserWindowPtr = WebBrowserWindowPtr;
+
 			[FIOSAsyncTask CreateTaskWithBlock : ^ bool(void)
 			{
-				TSharedPtr<FWebBrowserWindow> BrowserWindow = WebBrowserWindowPtr.Pin();
-				if (BrowserWindow.IsValid())
+				if (AsyncWebBrowserWindowPtr.IsValid())
 				{
-					TArray<FString> Params;
-					Message.ParseIntoArray(Params, TEXT("/"), false);
-					if (Params.Num() > 0)
+					TSharedPtr<FWebBrowserWindow> BrowserWindow = AsyncWebBrowserWindowPtr.Pin();
+					if (BrowserWindow.IsValid())
 					{
-						for (int I = 0; I < Params.Num(); I++)
+						TArray<FString> Params;
+						Message.ParseIntoArray(Params, TEXT("/"), false);
+						if (Params.Num() > 0)
 						{
-							Params[I] = FPlatformHttp::UrlDecode(Params[I]);
+							for (int I = 0; I < Params.Num(); I++)
+							{
+								Params[I] = FPlatformHttp::UrlDecode(Params[I]);
+							}
+							
+							FString Command = Params[0];
+							Params.RemoveAt(0, 1);
+							BrowserWindow->OnJsMessageReceived(Command, Params, "");
+						}
+						else
+						{
+							GLog->Logf(ELogVerbosity::Error, TEXT("Invalid message from browser view: %s"), *Message);
 						}
-
-						FString Command = Params[0];
-						Params.RemoveAt(0, 1);
-						BrowserWindow->OnJsMessageReceived(Command, Params, "");
-					}
-					else
-					{
-						GLog->Logf(ELogVerbosity::Error, TEXT("Invalid message from browser view: %s"), *Message);
 					}
 				}
 				return true;
@@ -789,14 +805,8 @@ supportsMetal : (bool)InSupportsMetal supportsMetalMRT : (bool)InSupportsMetalMR
 {
 	NSURLRequest *request = InNavigationAction.request;
 	FString UrlStr([[request URL]absoluteString]);
-
-	// Notify on the game thread
-	[FIOSAsyncTask CreateTaskWithBlock : ^ bool(void)
-	{
-		WebBrowserWidget->HandleShouldOverrideUrlLoading(UrlStr);
-		return true;
-	}];
-
+	
+	WebBrowserWidget->HandleShouldOverrideUrlLoading(UrlStr);
 	InDecisionHandler(WKNavigationActionPolicyAllow);
 }
 
