Pin() TWeakPtr before checking if valid due to possible race condition

diff --git a/Engine/Source/Runtime/WebBrowser/Private/IOS/IOSPlatformWebBrowser.cpp b/Engine/Source/Runtime/WebBrowser/Private/IOS/IOSPlatformWebBrowser.cpp
--- a/Engine/Source/Runtime/WebBrowser/Private/IOS/IOSPlatformWebBrowser.cpp
+++ b/Engine/Source/Runtime/WebBrowser/Private/IOS/IOSPlatformWebBrowser.cpp
@@ -315,20 +315,17 @@ class SIOSWebBrowserWidget : public SLeafWidget
 			// Notify on the game thread
 			[FIOSAsyncTask CreateTaskWithBlock : ^ bool(void)
 			{
-				if (AsyncWebBrowserWindowPtr.IsValid())
+				TSharedPtr<FWebBrowserWindow> BrowserWindow = AsyncWebBrowserWindowPtr.Pin();
+				if (AsyncWebBrowserWindowPtr.IsValid() && BrowserWindow.IsValid())
 				{
-					TSharedPtr<FWebBrowserWindow> BrowserWindow = AsyncWebBrowserWindowPtr.Pin();
-					if (BrowserWindow.IsValid())
+					if (BrowserWindow->OnBeforeBrowse().IsBound())
 					{
-						if (BrowserWindow->OnBeforeBrowse().IsBound())
-						{
-							FWebNavigationRequest RequestDetails;
-							RequestDetails.bIsRedirect = false;
-							RequestDetails.bIsMainFrame = true; // shouldOverrideUrlLoading is only called on the main frame
-							
-							BrowserWindow->OnBeforeBrowse().Execute(Url, RequestDetails);
-							BrowserWindow->SetTitle("");
-						}
+						FWebNavigationRequest RequestDetails;
+						RequestDetails.bIsRedirect = false;
+						RequestDetails.bIsMainFrame = true; // shouldOverrideUrlLoading is only called on the main frame
+
+						BrowserWindow->OnBeforeBrowse().Execute(Url, RequestDetails);
+						BrowserWindow->SetTitle("");
 					}
 				}
 				return true;
@@ -358,28 +355,25 @@ class SIOSWebBrowserWidget : public SLeafWidget
 
 			[FIOSAsyncTask CreateTaskWithBlock : ^ bool(void)
 			{
-				if (AsyncWebBrowserWindowPtr.IsValid())
+				TSharedPtr<FWebBrowserWindow> BrowserWindow = AsyncWebBrowserWindowPtr.Pin();
+				if (AsyncWebBrowserWindowPtr.IsValid() && BrowserWindow.IsValid())
 				{
-					TSharedPtr<FWebBrowserWindow> BrowserWindow = AsyncWebBrowserWindowPtr.Pin();
-					if (BrowserWindow.IsValid())
+					TArray<FString> Params;
+					Message.ParseIntoArray(Params, TEXT("/"), false);
+					if (Params.Num() > 0)
 					{
-						TArray<FString> Params;
-						Message.ParseIntoArray(Params, TEXT("/"), false);
-						if (Params.Num() > 0)
-						{
-							for (int I = 0; I < Params.Num(); I++)
-							{
-								Params[I] = FPlatformHttp::UrlDecode(Params[I]);
-							}
-							
-							FString Command = Params[0];
-							Params.RemoveAt(0, 1);
-							BrowserWindow->OnJsMessageReceived(Command, Params, "");
-						}
-						else
+						for (int I = 0; I < Params.Num(); I++)
 						{
-							GLog->Logf(ELogVerbosity::Error, TEXT("Invalid message from browser view: %s"), *Message);
+							Params[I] = FPlatformHttp::UrlDecode(Params[I]);
 						}
+
+						FString Command = Params[0];
+						Params.RemoveAt(0, 1);
+						BrowserWindow->OnJsMessageReceived(Command, Params, "");
+					}
+					else
+					{
+						GLog->Logf(ELogVerbosity::Error, TEXT("Invalid message from browser view: %s"), *Message);
 					}
 				}
 				return true;
