Simplify weakptr IsValid check

diff --git a/Engine/Source/Runtime/WebBrowser/Private/IOS/IOSPlatformWebBrowser.cpp b/Engine/Source/Runtime/WebBrowser/Private/IOS/IOSPlatformWebBrowser.cpp
--- a/Engine/Source/Runtime/WebBrowser/Private/IOS/IOSPlatformWebBrowser.cpp
+++ b/Engine/Source/Runtime/WebBrowser/Private/IOS/IOSPlatformWebBrowser.cpp
@@ -316,7 +316,7 @@ class SIOSWebBrowserWidget : public SLeafWidget
 			[FIOSAsyncTask CreateTaskWithBlock : ^ bool(void)
 			{
 				TSharedPtr<FWebBrowserWindow> BrowserWindow = AsyncWebBrowserWindowPtr.Pin();
-				if (AsyncWebBrowserWindowPtr.IsValid() && BrowserWindow.IsValid())
+				if (BrowserWindow.IsValid())
 				{
 					if (BrowserWindow->OnBeforeBrowse().IsBound())
 					{
@@ -356,7 +356,7 @@ class SIOSWebBrowserWidget : public SLeafWidget
 			[FIOSAsyncTask CreateTaskWithBlock : ^ bool(void)
 			{
 				TSharedPtr<FWebBrowserWindow> BrowserWindow = AsyncWebBrowserWindowPtr.Pin();
-				if (AsyncWebBrowserWindowPtr.IsValid() && BrowserWindow.IsValid())
+				if (BrowserWindow.IsValid())
 				{
 					TArray<FString> Params;
 					Message.ParseIntoArray(Params, TEXT("/"), false);
