Let Android WebBrowser handle mouse (forward to webview as touch events)

#android

diff --git a/Engine/Build/Android/Java/src/com/epicgames/unreal/WebViewControl.java b/Engine/Build/Android/Java/src/com/epicgames/unreal/WebViewControl.java
--- a/Engine/Build/Android/Java/src/com/epicgames/unreal/WebViewControl.java
+++ b/Engine/Build/Android/Java/src/com/epicgames/unreal/WebViewControl.java
@@ -102,6 +102,8 @@ class WebViewControl
 
 		nativePtr = inNativePtr;
 
+		//GameActivity.Log.debug("WebViewControl width=" + width + ", height=" + height);
+
 		GameActivity._activity.runOnUiThread(new Runnable()
 		{
 			@Override
@@ -302,6 +304,30 @@ class WebViewControl
 		return webView.canGoBackOrForward(Steps);
 	}
 
+	public void SendTouchEvent(int event, float x, float y)
+	{
+		final int actionType = event;
+		final long actionTime = SystemClock.uptimeMillis();
+		final float actionX = webView.getLeft() + (x * webView.getWidth());
+		final float actionY = webView.getTop() + (y * webView.getHeight());
+		//GameActivity.Log.debug("SendTouchEvent(event=" + event + ", x=" + x + ", y=" + y + ") = " + actionX + ", " + actionY);
+		GameActivity._activity.runOnUiThread(new Runnable()
+		{
+			@Override
+			public void run()
+			{
+				long eventTime = SystemClock.uptimeMillis();
+				webView.onTouchEvent(MotionEvent.obtain(
+					actionTime,
+					eventTime,
+					actionType,
+					actionX,
+					actionY,
+					0
+				));
+			}
+		});
+	}
 
 	// called from C++ paint event
 	public void Update(final int x, final int y, final int width, final int height)
diff --git a/Engine/Source/Runtime/UMG/Private/Components/WidgetInteractionComponent.cpp b/Engine/Source/Runtime/UMG/Private/Components/WidgetInteractionComponent.cpp
--- a/Engine/Source/Runtime/UMG/Private/Components/WidgetInteractionComponent.cpp
+++ b/Engine/Source/Runtime/UMG/Private/Components/WidgetInteractionComponent.cpp
@@ -524,27 +524,49 @@ void UWidgetInteractionComponent::ReleasePointerKey(FKey Key)
 
 	ensure(PointerIndex >= 0);
 	FPointerEvent PointerEvent;
+
+	// Find the primary input device for this Slate User
+	FInputDeviceId InputDeviceId = INPUTDEVICEID_NONE;
+	if (TSharedPtr<FSlateUser> SlateUser = FSlateApplication::Get().GetUser(VirtualUser->GetUserIndex()))
+	{
+		FPlatformUserId PlatUser = SlateUser->GetPlatformUserId();
+		InputDeviceId = IPlatformInputDeviceMapper::Get().GetPrimaryInputDeviceForUser(PlatUser);
+	}
+
+	// Just in case there was no input device assigned to this virtual user, get the default platform
+	// input device
+	if (!InputDeviceId.IsValid())
+	{
+		InputDeviceId = IPlatformInputDeviceMapper::Get().GetDefaultInputDevice();
+	}
+
 	if (Key.IsTouch())
 	{
 		PointerEvent = FPointerEvent(
-			VirtualUser->GetUserIndex(),
+			InputDeviceId,
 			(uint32)PointerIndex,
 			LocalHitLocation,
 			LastLocalHitLocation,
-			1.0f,
-			false);
+			0.0f,
+			false,
+			false,
+			false,
+			FModifierKeysState(),
+			0,
+			VirtualUser->GetUserIndex());
 	}
 	else
 	{
 		PointerEvent = FPointerEvent(
-			VirtualUser->GetUserIndex(),
+			InputDeviceId,
 			(uint32)PointerIndex,
 			LocalHitLocation,
 			LastLocalHitLocation,
 			PressedKeys,
 			Key,
 			0.0f,
-			ModifierKeys);
+			ModifierKeys,
+			VirtualUser->GetUserIndex());
 	}
 		
 	FReply Reply = FSlateApplication::Get().RoutePointerUpEvent(WidgetPathUnderFinger, PointerEvent);
diff --git a/Engine/Source/Runtime/WebBrowser/Private/Android/AndroidJavaWebBrowser.cpp b/Engine/Source/Runtime/WebBrowser/Private/Android/AndroidJavaWebBrowser.cpp
--- a/Engine/Source/Runtime/WebBrowser/Private/Android/AndroidJavaWebBrowser.cpp
+++ b/Engine/Source/Runtime/WebBrowser/Private/Android/AndroidJavaWebBrowser.cpp
@@ -25,6 +25,11 @@ static jfieldID FindField(JNIEnv* JEnv, jclass Class, const ANSICHAR* FieldName,
 	return Field;
 }
 
+// From https://developer.android.com/reference/android/view/MotionEvent#constants_1
+#define MOTIONEVENT_ACTION_DOWN		0
+#define MOTIONEVENT_ACTION_UP		1
+#define MOTIONEVENT_ACTION_MOVE		2
+
 FJavaAndroidWebBrowser::FJavaAndroidWebBrowser(bool swizzlePixels, bool vulkanRenderer, int32 width, int32 height,
 	jlong widgetPtr, bool bEnableRemoteDebugging, bool bUseTransparency, bool bEnableDomStorage)
 	: FJavaClassObject(GetClassName(), "(JIIZZZZZ)V", widgetPtr, width, height,  swizzlePixels, vulkanRenderer, bEnableRemoteDebugging, bUseTransparency, bEnableDomStorage)
@@ -41,6 +46,7 @@ FJavaAndroidWebBrowser::FJavaAndroidWebBrowser(bool swizzlePixels, bool vulkanRe
 	, ReloadMethod(GetClassMethod("Reload", "()V"))
 	, CloseMethod(GetClassMethod("Close", "()V"))
 	, GoBackOrForwardMethod(GetClassMethod("GoBackOrForward", "(I)V"))
+	, SendTouchEventMethod(GetClassMethod("SendTouchEvent", "(IFF)V"))
 	, SetAndroid3DBrowserMethod(GetClassMethod("SetAndroid3DBrowser", "(Z)V"))
 	, SetVisibilityMethod(GetClassMethod("SetVisibility", "(Z)V"))
 {
@@ -208,6 +214,21 @@ void FJavaAndroidWebBrowser::GoForward()
 	CallMethod<void>(GoBackOrForwardMethod, 1);
 }
 
+void FJavaAndroidWebBrowser::SendTouchDown(float x, float y)
+{
+	CallMethod<void>(SendTouchEventMethod, MOTIONEVENT_ACTION_DOWN, x, y);
+}
+
+void FJavaAndroidWebBrowser::SendTouchUp(float x, float y)
+{
+	CallMethod<void>(SendTouchEventMethod, MOTIONEVENT_ACTION_UP, x, y);
+}
+
+void FJavaAndroidWebBrowser::SendTouchMove(float x, float y)
+{
+	CallMethod<void>(SendTouchEventMethod, MOTIONEVENT_ACTION_MOVE, x, y);
+}
+
 void FJavaAndroidWebBrowser::SetAndroid3DBrowser(bool InIsAndroid3DBrowser)
 {
 	CallMethod<void>(SetAndroid3DBrowserMethod, InIsAndroid3DBrowser);
diff --git a/Engine/Source/Runtime/WebBrowser/Private/Android/AndroidJavaWebBrowser.h b/Engine/Source/Runtime/WebBrowser/Private/Android/AndroidJavaWebBrowser.h
--- a/Engine/Source/Runtime/WebBrowser/Private/Android/AndroidJavaWebBrowser.h
+++ b/Engine/Source/Runtime/WebBrowser/Private/Android/AndroidJavaWebBrowser.h
@@ -30,6 +30,9 @@ public:
 	void Close();
 	void GoBack();
 	void GoForward();
+	void SendTouchDown(float x, float y);
+	void SendTouchUp(float x, float y);
+	void SendTouchMove(float x, float y);
 	void SetAndroid3DBrowser(bool InIsAndroid3DBrowser);
 	void SetVisibility(bool InIsVisible);
 	void Update(const int posX, const int posY, const int sizeX, const int sizeY);
@@ -49,6 +52,7 @@ private:
 	FJavaClassMethod ReloadMethod;
 	FJavaClassMethod CloseMethod;
 	FJavaClassMethod GoBackOrForwardMethod;
+	FJavaClassMethod SendTouchEventMethod;
 	FJavaClassMethod SetAndroid3DBrowserMethod;
 	FJavaClassMethod SetVisibilityMethod;
 
diff --git a/Engine/Source/Runtime/WebBrowser/Private/Android/AndroidWebBrowserWidget.cpp b/Engine/Source/Runtime/WebBrowser/Private/Android/AndroidWebBrowserWidget.cpp
--- a/Engine/Source/Runtime/WebBrowser/Private/Android/AndroidWebBrowserWidget.cpp
+++ b/Engine/Source/Runtime/WebBrowser/Private/Android/AndroidWebBrowserWidget.cpp
@@ -120,6 +120,8 @@ void SAndroidWebBrowserWidget::Construct(const FArguments& Args)
 	GConfig->GetBool(TEXT("/Script/AndroidRuntimeSettings.AndroidRuntimeSettings"), TEXT("bEnableDomStorage"), bEnableDomStorage, GEngineIni);
 
 	FIntPoint viewportSize = WebBrowserWindowPtr.Pin()->GetViewportSize();
+	FPlatformMisc::LowLevelOutputDebugStringf(TEXT("SAndroidWebBrowserWidget::Construct viewport=%d x %d"), viewportSize.X, viewportSize.Y);
+
 	JavaWebBrowser = MakeShared<FJavaAndroidWebBrowser, ESPMode::ThreadSafe>(false, FAndroidMisc::ShouldUseVulkan(), viewportSize.X, viewportSize.Y,
 		reinterpret_cast<jlong>(this), !(UE_BUILD_SHIPPING || UE_BUILD_TEST), Args._UseTransparency, bEnableDomStorage);
 
@@ -196,31 +198,49 @@ void SAndroidWebBrowserWidget::Tick(const FGeometry& AllottedGeometry, const dou
 		JavaWebBrowser->SetVideoTextureValid(false);
 	}
 
-	// Calculate UIScale, which can vary frame-to-frame thanks to device rotation
-	// UI Scale is calculated relative to vertical axis of 1280x720 / 720x1280
-	float UIScale;
-	FPlatformRect ScreenRect = FAndroidWindow::GetScreenRect();
-	int32_t ScreenWidth, ScreenHeight;
-	FAndroidWindow::CalculateSurfaceSize(ScreenWidth, ScreenHeight);
-	if (ScreenWidth > ScreenHeight)
+
+	FIntPoint viewportSize = WebBrowserWindowPtr.Pin()->GetViewportSize();
+
+	if (IsAndroid3DBrowser)
 	{
-		UIScale = (float)ScreenHeight / (ScreenRect.Bottom - ScreenRect.Top);
+		//FVector2D LocalSize = AllottedGeometry.GetLocalSize();
+		//FIntPoint IntLocalSize = FIntPoint((int32)LocalSize.X, (int32)LocalSize.Y);
+
+		//if (viewportSize != IntLocalSize)
+		//{
+		//	FPlatformMisc::LowLevelOutputDebugStringf(TEXT("SAndroidWebBrowser::Tick: updating viewport to localSize = %d x %d"), IntLocalSize.X, IntLocalSize.Y);
+		//	WebBrowserWindowPtr.Pin()->SetViewportSize(IntLocalSize, FIntPoint(0, 0));
+		//}
+
+		JavaWebBrowser->Update(0, 0, viewportSize.X, viewportSize.Y);
 	}
 	else
 	{
-		UIScale = (float)ScreenHeight / (ScreenRect.Bottom - ScreenRect.Top);
-	}
-
-	FVector2D Position = AllottedGeometry.GetAccumulatedRenderTransform().GetTranslation() * UIScale;
-	FVector2D Size = TransformVector(AllottedGeometry.GetAccumulatedRenderTransform(), AllottedGeometry.GetLocalSize()) * UIScale;
+		// Calculate UIScale, which can vary frame-to-frame thanks to device rotation
+		// UI Scale is calculated relative to vertical axis of 1280x720 / 720x1280
+		float UIScale;
+		FPlatformRect ScreenRect = FAndroidWindow::GetScreenRect();
+		int32_t ScreenWidth, ScreenHeight;
+		FAndroidWindow::CalculateSurfaceSize(ScreenWidth, ScreenHeight);
+		if (ScreenWidth > ScreenHeight)
+		{
+			UIScale = (float)ScreenHeight / (ScreenRect.Bottom - ScreenRect.Top);
+		}
+		else
+		{
+			UIScale = (float)ScreenHeight / (ScreenRect.Bottom - ScreenRect.Top);
+		}
 
-	// Convert position to integer coordinates
-	FIntPoint IntPos(FMath::RoundToInt(Position.X), FMath::RoundToInt(Position.Y));
-	// Convert size to integer taking the rounding of position into account to avoid double round-down or double round-up causing a noticeable error.
-	FIntPoint IntSize = FIntPoint(FMath::RoundToInt(Position.X + Size.X), FMath::RoundToInt(Size.Y + Position.Y)) - IntPos;
+		FVector2D Position = AllottedGeometry.GetAccumulatedRenderTransform().GetTranslation() * UIScale;
+		FVector2D Size = TransformVector(AllottedGeometry.GetAccumulatedRenderTransform(), AllottedGeometry.GetLocalSize()) * UIScale;
 
-	JavaWebBrowser->Update(IntPos.X, IntPos.Y, IntSize.X, IntSize.Y);
+		// Convert position to integer coordinates
+		FIntPoint IntPos(FMath::RoundToInt(Position.X), FMath::RoundToInt(Position.Y));
+		// Convert size to integer taking the rounding of position into account to avoid double round-down or double round-up causing a noticeable error.
+		FIntPoint IntSize = FIntPoint(FMath::RoundToInt(Position.X + Size.X), FMath::RoundToInt(Size.Y + Position.Y)) - IntPos;
 
+		JavaWebBrowser->Update(IntPos.X, IntPos.Y, IntSize.X, IntSize.Y);
+	}
 
 	if (IsAndroid3DBrowser)
 	{
@@ -237,8 +257,6 @@ void SAndroidWebBrowserWidget::Tick(const FGeometry& AllottedGeometry, const dou
 			// create new video sample
 			auto NewTextureSample = TextureSamplePool->AcquireShared();
 
-			FIntPoint viewportSize = WebBrowserWindowPtr.Pin()->GetViewportSize();
-
 			if (!NewTextureSample->Initialize(viewportSize))
 			{
 				return;
@@ -272,7 +290,7 @@ void SAndroidWebBrowserWidget::Tick(const FGeometry& AllottedGeometry, const dou
 
 					if (!PinnedJavaWebBrowser->GetVideoLastFrameData(Buffer, SampleCount, &bRegionChanged))
 					{
-						FPlatformMisc::LowLevelOutputDebugStringf(TEXT("Fetch RT: ShouldUseVulkan couldn't get texture buffer"));
+						//FPlatformMisc::LowLevelOutputDebugStringf(TEXT("Fetch RT: ShouldUseVulkan couldn't get texture buffer"));
 						return;
 					}
 
@@ -298,8 +316,6 @@ void SAndroidWebBrowserWidget::Tick(const FGeometry& AllottedGeometry, const dou
 				FIntPoint Size;
 			};
 
-			FIntPoint viewportSize = WebBrowserWindowPtr.Pin()->GetViewportSize();
-
 			FWriteWebBrowserParams WriteWebBrowserParams = { JavaWebBrowser, WebBrowserTexture->GetExternalTextureGuid(), viewportSize };
 			ENQUEUE_RENDER_COMMAND(WriteAndroidWebBrowser)(
 				[Params = WriteWebBrowserParams](FRHICommandListImmediate& RHICmdList)
@@ -362,8 +378,6 @@ void SAndroidWebBrowserWidget::Tick(const FGeometry& AllottedGeometry, const dou
 			// create new video sample
 			auto NewTextureSample = TextureSamplePool->AcquireShared();
 
-			FIntPoint viewportSize = WebBrowserWindowPtr.Pin()->GetViewportSize();
-
 			if (!NewTextureSample->Initialize(viewportSize))
 			{
 				return;
@@ -477,6 +491,82 @@ bool SAndroidWebBrowserWidget::CanGoForward()
 	return HistoryPosition < HistorySize-1;
 }
 
+void SAndroidWebBrowserWidget::SendTouchDown(FVector2D Position)
+{
+	FVector2D WidgetSize = GetCachedGeometry().GetLocalSize();
+	JavaWebBrowser->SendTouchDown(Position.X / WidgetSize.X, Position.Y / WidgetSize.Y);
+}
+
+void SAndroidWebBrowserWidget::SendTouchUp(FVector2D Position)
+{
+	FVector2D WidgetSize = GetCachedGeometry().GetLocalSize();
+	JavaWebBrowser->SendTouchUp(Position.X / WidgetSize.X, Position.Y / WidgetSize.Y);
+}
+
+void SAndroidWebBrowserWidget::SendTouchMove(FVector2D Position)
+{
+	FVector2D WidgetSize = GetCachedGeometry().GetLocalSize();
+	JavaWebBrowser->SendTouchMove(Position.X / WidgetSize.X, Position.Y / WidgetSize.Y);
+}
+
+FVector2D SAndroidWebBrowserWidget::ConvertMouseEventToLocal(const FGeometry& MyGeometry, const FPointerEvent& MouseEvent)
+{
+	FGeometry MouseGeometry = MyGeometry;
+
+	float DPIScale = MouseGeometry.Scale;
+	FVector2D LocalPos = MouseGeometry.AbsoluteToLocal(MouseEvent.GetScreenSpacePosition()) * DPIScale;
+
+	return LocalPos;
+}
+
+FReply SAndroidWebBrowserWidget::OnMouseButtonDown(const FGeometry& MyGeometry, const FPointerEvent& MouseEvent)
+{
+	FReply Reply = FReply::Unhandled();
+
+	FKey Button = MouseEvent.GetEffectingButton();
+	bool bSupportedButton = (Button == EKeys::LeftMouseButton); // || Button == EKeys::RightMouseButton || Button == EKeys::MiddleMouseButton);
+
+	if (bSupportedButton)
+	{
+		Reply = FReply::Handled();
+		SendTouchDown(ConvertMouseEventToLocal(MyGeometry, MouseEvent));
+	}
+
+	return Reply;
+}
+
+FReply SAndroidWebBrowserWidget::OnMouseButtonUp(const FGeometry& MyGeometry, const FPointerEvent& MouseEvent)
+{
+	FReply Reply = FReply::Unhandled();
+
+	FKey Button = MouseEvent.GetEffectingButton();
+	bool bSupportedButton = (Button == EKeys::LeftMouseButton); // || Button == EKeys::RightMouseButton || Button == EKeys::MiddleMouseButton);
+
+	if (bSupportedButton)
+	{
+		Reply = FReply::Handled();
+		SendTouchUp(ConvertMouseEventToLocal(MyGeometry, MouseEvent));
+	}
+
+	return Reply;
+}
+
+FReply SAndroidWebBrowserWidget::OnMouseMove(const FGeometry& MyGeometry, const FPointerEvent& MouseEvent)
+{
+	FReply Reply = FReply::Unhandled();
+
+	FKey Button = MouseEvent.GetEffectingButton();
+	bool bSupportedButton = (Button == EKeys::LeftMouseButton); // || Button == EKeys::RightMouseButton || Button == EKeys::MiddleMouseButton);
+
+	if (bSupportedButton)
+	{
+		Reply = FReply::Handled();
+		SendTouchMove(ConvertMouseEventToLocal(MyGeometry, MouseEvent));
+	}
+
+	return Reply;
+}
+
 void SAndroidWebBrowserWidget::SetWebBrowserVisibility(bool InIsVisible)
 {
 	JavaWebBrowser->SetVisibility(InIsVisible);
diff --git a/Engine/Source/Runtime/WebBrowser/Private/Android/AndroidWebBrowserWidget.h b/Engine/Source/Runtime/WebBrowser/Private/Android/AndroidWebBrowserWidget.h
--- a/Engine/Source/Runtime/WebBrowser/Private/Android/AndroidWebBrowserWidget.h
+++ b/Engine/Source/Runtime/WebBrowser/Private/Android/AndroidWebBrowserWidget.h
@@ -59,6 +59,14 @@ public:
 	bool CanGoBack();
 	bool CanGoForward();
 
+	virtual FReply OnMouseButtonDown(const FGeometry& MyGeometry, const FPointerEvent& MouseEvent) override;
+	virtual FReply OnMouseButtonUp(const FGeometry& MyGeometry, const FPointerEvent& MouseEvent) override;
+	virtual FReply OnMouseMove(const FGeometry& MyGeometry, const FPointerEvent& MouseEvent) override;
+
+	void SendTouchDown(FVector2D Position);
+	void SendTouchUp(FVector2D Position);
+	void SendTouchMove(FVector2D Position);
+
 	// WebViewClient callbacks
 
 	jbyteArray HandleShouldInterceptRequest(jstring JUrl);
@@ -96,6 +104,7 @@ protected:
 
 	TWeakPtr<FAndroidWebBrowserWindow> WebBrowserWindowPtr;
 private:
+	FVector2D ConvertMouseEventToLocal(const FGeometry& MyGeometry, const FPointerEvent& MouseEvent);
 
 	/** Enable 3D appearance for Android. */
 	bool IsAndroid3DBrowser;
diff --git a/Engine/Source/Runtime/WebBrowser/Private/Android/AndroidWebBrowserWindow.cpp b/Engine/Source/Runtime/WebBrowser/Private/Android/AndroidWebBrowserWindow.cpp
--- a/Engine/Source/Runtime/WebBrowser/Private/Android/AndroidWebBrowserWindow.cpp
+++ b/Engine/Source/Runtime/WebBrowser/Private/Android/AndroidWebBrowserWindow.cpp
@@ -29,7 +29,7 @@ FAndroidWebBrowserWindow::FAndroidWebBrowserWindow(FString InUrl, TOptional<FStr
 	, DocumentState(EWebBrowserDocumentState::NoDocument)
 	, ErrorCode(0)
 	, Scripting(new FMobileJSScripting(bInJSBindingToLoweringEnabled))
-	, AndroidWindowSize(FIntPoint(500, 500))
+	, AndroidWindowSize(FIntPoint(1024, 768))
 	, bIsDisabled(false)
 	, bIsVisible(true)
 	, bTickedLastFrame(true)
@@ -83,7 +83,7 @@ FSlateShaderResource* FAndroidWebBrowserWindow::GetTexture(bool bIsPopup /*= fal
 
 bool FAndroidWebBrowserWindow::IsValid() const
 {
-	return false;
+	return BrowserWidget.IsValid();
 }
 
 bool FAndroidWebBrowserWindow::IsInitialized() const
@@ -126,14 +126,46 @@ bool FAndroidWebBrowserWindow::OnKeyChar(const FCharacterEvent& InCharacterEvent
 	return false;
 }
 
+FVector2D FAndroidWebBrowserWindow::ConvertMouseEventToLocal(const FGeometry& MyGeometry, const FPointerEvent& MouseEvent, bool bIsPopup)
+{
+	FGeometry MouseGeometry = MyGeometry;
+
+	float DPIScale = MouseGeometry.Scale;
+	FVector2D LocalPos = MouseGeometry.AbsoluteToLocal(MouseEvent.GetScreenSpacePosition()) * DPIScale;
+
+	return LocalPos;
+}
+
 FReply FAndroidWebBrowserWindow::OnMouseButtonDown(const FGeometry& MyGeometry, const FPointerEvent& MouseEvent, bool bIsPopup)
 {
-	return FReply::Unhandled();
+	FReply Reply = FReply::Unhandled();
+
+	FKey Button = MouseEvent.GetEffectingButton();
+	bool bSupportedButton = (Button == EKeys::LeftMouseButton); // || Button == EKeys::RightMouseButton || Button == EKeys::MiddleMouseButton);
+
+	if (bSupportedButton)
+	{
+		Reply = FReply::Handled();
+		BrowserWidget->SendTouchDown(ConvertMouseEventToLocal(MyGeometry, MouseEvent, bIsPopup));
+	}
+
+	return Reply;
 }
 
 FReply FAndroidWebBrowserWindow::OnMouseButtonUp(const FGeometry& MyGeometry, const FPointerEvent& MouseEvent, bool bIsPopup)
 {
-	return FReply::Unhandled();
+	FReply Reply = FReply::Unhandled();
+
+	FKey Button = MouseEvent.GetEffectingButton();
+	bool bSupportedButton = (Button == EKeys::LeftMouseButton); // || Button == EKeys::RightMouseButton || Button == EKeys::MiddleMouseButton);
+
+	if (bSupportedButton)
+	{
+		Reply = FReply::Handled();
+		BrowserWidget->SendTouchUp(ConvertMouseEventToLocal(MyGeometry, MouseEvent, bIsPopup));
+	}
+
+	return Reply;
 }
 
 FReply FAndroidWebBrowserWindow::OnMouseButtonDoubleClick(const FGeometry& MyGeometry, const FPointerEvent& MouseEvent, bool bIsPopup)
@@ -143,7 +175,18 @@ FReply FAndroidWebBrowserWindow::OnMouseButtonDoubleClick(const FGeometry& MyGeo
 
 FReply FAndroidWebBrowserWindow::OnMouseMove(const FGeometry& MyGeometry, const FPointerEvent& MouseEvent, bool bIsPopup)
 {
-	return FReply::Unhandled();
+	FReply Reply = FReply::Unhandled();
+
+	FKey Button = MouseEvent.GetEffectingButton();
+	bool bSupportedButton = (Button == EKeys::LeftMouseButton); // || Button == EKeys::RightMouseButton || Button == EKeys::MiddleMouseButton);
+
+	if (bSupportedButton)
+	{
+		Reply = FReply::Handled();
+		BrowserWidget->SendTouchMove(ConvertMouseEventToLocal(MyGeometry, MouseEvent, bIsPopup));
+	}
+
+	return Reply;
 }
 
 void FAndroidWebBrowserWindow::OnMouseLeave(const FPointerEvent& MouseEvent)
diff --git a/Engine/Source/Runtime/WebBrowser/Private/Android/AndroidWebBrowserWindow.h b/Engine/Source/Runtime/WebBrowser/Private/Android/AndroidWebBrowserWindow.h
--- a/Engine/Source/Runtime/WebBrowser/Private/Android/AndroidWebBrowserWindow.h
+++ b/Engine/Source/Runtime/WebBrowser/Private/Android/AndroidWebBrowserWindow.h
@@ -246,6 +246,8 @@ private:
 
 	bool OnJsMessageReceived(const FString& Command, const TArray<FString>& Params, const FString& Origin);
 
+	FVector2D ConvertMouseEventToLocal(const FGeometry& MyGeometry, const FPointerEvent& MouseEvent, bool bIsPopup);
+
 	TSharedPtr<SAndroidWebBrowserWidget> BrowserWidget;
 
 	/** Current title of this window. */
diff --git a/Engine/Source/Runtime/WebBrowserTexture/Private/WebBrowserTextureResource.cpp b/Engine/Source/Runtime/WebBrowserTexture/Private/WebBrowserTextureResource.cpp
--- a/Engine/Source/Runtime/WebBrowserTexture/Private/WebBrowserTextureResource.cpp
+++ b/Engine/Source/Runtime/WebBrowserTexture/Private/WebBrowserTextureResource.cpp
@@ -17,6 +17,7 @@
 
 #define WebBrowserTextureRESOURCE_TRACE_RENDER 0
 
+DEFINE_LOG_CATEGORY(LogWebBrowserTexture);
 
 /* FWebBrowserTextureResource structors
  *****************************************************************************/
@@ -27,9 +28,8 @@ FWebBrowserTextureResource::FWebBrowserTextureResource(UWebBrowserTexture& InOwn
 	, Owner(InOwner)
 	, OwnerDim(InOwnerDim)
 	, OwnerSize(InOwnerSize)
-{ 
-	FPlatformMisc::LowLevelOutputDebugStringf(TEXT("FWebBrowserTextureResource:FWebBrowserTextureResource %d %d"), OwnerDim.X, OwnerDim.Y);
-
+{
+	UE_LOG(LogWebBrowserTexture, VeryVerbose, TEXT("FWebBrowserTextureResource:FWebBrowserTextureResource %d %d"), OwnerDim.X, OwnerDim.Y);
 }
 
 
@@ -144,7 +144,7 @@ void FWebBrowserTextureResource::ReleaseDynamicRHI()
 
 void FWebBrowserTextureResource::ClearTexture(const FLinearColor& ClearColor)
 {
-	FPlatformMisc::LowLevelOutputDebugStringf(TEXT("FWebBrowserTextureResource:ClearTexture"));
+	UE_LOG(LogWebBrowserTexture, VeryVerbose, TEXT("FWebBrowserTextureResource:ClearTexture"));
 	// create output render target if we don't have one yet
 	const ETextureCreateFlags OutputCreateFlags = ETextureCreateFlags::Dynamic | ETextureCreateFlags::SRGB;
 
@@ -189,11 +189,11 @@ void FWebBrowserTextureResource::CopySample(const TSharedPtr<FWebBrowserTextureS
 	// we create our own output render target and copy the data into it.
 	if (SampleTexture2D != nullptr)
 	{
-	FPlatformMisc::LowLevelOutputDebugStringf(TEXT("FWebBrowserTextureResource:CopySample 1"));
+		UE_LOG(LogWebBrowserTexture, VeryVerbose, TEXT("FWebBrowserTextureResource:CopySample 1"));
 		// use sample's texture as the new render target.
 		if (TextureRHI != SampleTexture2D)
 		{
-			FPlatformMisc::LowLevelOutputDebugStringf(TEXT("FWebBrowserTextureResource:CopySample 11"));
+			UE_LOG(LogWebBrowserTexture, VeryVerbose, TEXT("FWebBrowserTextureResource:CopySample 11"));
 			UpdateTextureReference(SampleTexture2D);
 
 			OutputTarget.SafeRelease();
@@ -202,14 +202,14 @@ void FWebBrowserTextureResource::CopySample(const TSharedPtr<FWebBrowserTextureS
 	}
 	else
 	{
-		FPlatformMisc::LowLevelOutputDebugStringf(TEXT("FWebBrowserTextureResource:CopySample 2"));
+		UE_LOG(LogWebBrowserTexture, VeryVerbose, TEXT("FWebBrowserTextureResource:CopySample 2"));
 		// create a new output render target if necessary
 		const ETextureCreateFlags OutputCreateFlags = TexCreate_Dynamic | TexCreate_SRGB;
 		const FIntPoint SampleDim = Sample->GetDim();
 
 		if ((ClearColor != CurrentClearColor) || !OutputTarget.IsValid() || (OutputTarget->GetSizeXY() != SampleDim) || ((OutputTarget->GetFlags() & OutputCreateFlags) != OutputCreateFlags))
 		{
-			FPlatformMisc::LowLevelOutputDebugStringf(TEXT("FWebBrowserTextureResource:CopySample 1"));
+			UE_LOG(LogWebBrowserTexture, VeryVerbose, TEXT("FWebBrowserTextureResource:CopySample 1"));
 
 			const FRHITextureCreateDesc Desc =
 				FRHITextureCreateDesc::Create2D(TEXT("FWebBrowserTextureResource"))
@@ -230,6 +230,8 @@ void FWebBrowserTextureResource::CopySample(const TSharedPtr<FWebBrowserTextureS
 			UpdateTextureReference(OutputTarget);
 		}
 
+		UE_LOG(LogWebBrowserTexture, VeryVerbose, TEXT("WebBrowserTextureResource:CopySample: %d x %d"), SampleDim.X, SampleDim.Y);
+
 		// copy sample data to output render target
 		FUpdateTextureRegion2D Region(0, 0, 0, 0, SampleDim.X, SampleDim.Y);
 		RHIUpdateTexture2D(RenderTargetTextureRHI.GetReference(), 0, Region, Sample->GetStride(), (uint8*)Sample->GetBuffer());
@@ -240,7 +242,7 @@ void FWebBrowserTextureResource::CopySample(const TSharedPtr<FWebBrowserTextureS
 
 void FWebBrowserTextureResource::UpdateResourceSize()
 {
-	FPlatformMisc::LowLevelOutputDebugStringf(TEXT("FWebBrowserTextureResource:UpdateResourceSize"));
+	UE_LOG(LogWebBrowserTexture, VeryVerbose, TEXT("FWebBrowserTextureResource:UpdateResourceSize"));
 
 	SIZE_T ResourceSize = 0;
 
@@ -260,7 +262,6 @@ void FWebBrowserTextureResource::UpdateResourceSize()
 
 void FWebBrowserTextureResource::UpdateTextureReference(FRHITexture2D* NewTexture)
 {
-	FPlatformMisc::LowLevelOutputDebugStringf(TEXT("FWebBrowserTextureResource:UpdateTextureReference"));
 	TextureRHI = NewTexture;
 	RenderTargetTextureRHI = NewTexture;
 
@@ -274,4 +275,5 @@ void FWebBrowserTextureResource::UpdateTextureReference(FRHITexture2D* NewTextur
 	{
 		OwnerDim = FIntPoint::ZeroValue;
 	}
+	UE_LOG(LogWebBrowserTexture, VeryVerbose, TEXT("FWebBrowserTextureResource:UpdateTextureReference: %d x %d"), OwnerDim.X, OwnerDim.Y);
 }
diff --git a/Engine/Source/Runtime/WebBrowserTexture/Public/WebBrowserTextureResource.h b/Engine/Source/Runtime/WebBrowserTexture/Public/WebBrowserTextureResource.h
--- a/Engine/Source/Runtime/WebBrowserTexture/Public/WebBrowserTextureResource.h
+++ b/Engine/Source/Runtime/WebBrowserTexture/Public/WebBrowserTextureResource.h
@@ -4,6 +4,7 @@
 
 #include "CoreTypes.h"
 #include "Containers/Queue.h"
+#include "Logging/LogMacros.h"
 #include "Math/Color.h"
 #include "Misc/Guid.h"
 #include "Misc/Timespan.h"
@@ -15,6 +16,8 @@
 
 class UWebBrowserTexture;
 
+WEBBROWSERTEXTURE_API DECLARE_LOG_CATEGORY_EXTERN(LogWebBrowserTexture, Log, All);
+
 /**
  * Texture resource type for webbrowser textures.
  */
