Choose Android WebBrowser texture size in DefaultEngine.ini

[/Script/AndroidRuntimeSettings.AndroidRuntimeSettings]
WebViewTextureSize=1024,768

#android

diff --git a/Engine/Source/Runtime/WebBrowser/Private/Android/AndroidWebBrowserWindow.cpp b/Engine/Source/Runtime/WebBrowser/Private/Android/AndroidWebBrowserWindow.cpp
--- a/Engine/Source/Runtime/WebBrowser/Private/Android/AndroidWebBrowserWindow.cpp
+++ b/Engine/Source/Runtime/WebBrowser/Private/Android/AndroidWebBrowserWindow.cpp
@@ -9,6 +9,7 @@
 #include "Android/AndroidApplication.h"
 #include "Android/AndroidWindow.h"
 #include "Android/AndroidJava.h"
+#include "Misc/ConfigCacheIni.h"
 
 #include <jni.h>
 
@@ -34,6 +35,23 @@ FAndroidWebBrowserWindow::FAndroidWebBrowserWindow(FString InUrl, TOptional<FStr
 	, bIsVisible(true)
 	, bTickedLastFrame(true)
 {
+	// Deal with optional texture resolution override
+	FString WebViewTextureSize;
+	if (GConfig->GetString(TEXT("/Script/AndroidRuntimeSettings.AndroidRuntimeSettings"), TEXT("WebViewTextureSize"), WebViewTextureSize, GEngineIni))
+	{
+		TArray<FString> ResolutionVector;
+		int ParseCount = WebViewTextureSize.ParseIntoArray(ResolutionVector, TEXT(","), false);
+		ensureMsgf(ParseCount == 2, TEXT("WebViewTextureSize variable not properly formatted."));
+
+		if (ParseCount == 2)
+		{
+			int32 Width = FCString::Atof(*ResolutionVector[0]);
+			int32 Height = FCString::Atof(*ResolutionVector[1]);
+
+			AndroidWindowSize.X = FMath::Max(1, Width);
+			AndroidWindowSize.Y = FMath::Max(1, Height);
+		}
+	}
 }
 
 FAndroidWebBrowserWindow::~FAndroidWebBrowserWindow()
