- Fix SWebBrowserView::HandleBrowserWindowNeedsRedraw assuming the SlateApplication was still alive
- Add assert to FSlateApplication::Get that the application is alive.

diff --git a/Engine/Source/Runtime/Slate/Public/Framework/Application/SlateApplication.h b/Engine/Source/Runtime/Slate/Public/Framework/Application/SlateApplication.h
--- a/Engine/Source/Runtime/Slate/Public/Framework/Application/SlateApplication.h
+++ b/Engine/Source/Runtime/Slate/Public/Framework/Application/SlateApplication.h
@@ -252,6 +252,7 @@ public:
 	 */
 	static FSlateApplication& Get()
 	{
+		check( CurrentApplication.IsValid() );
 		check( IsInGameThread() || IsInSlateThread() || IsInAsyncLoadingThread() );
 		return *CurrentApplication;
 	}
diff --git a/Engine/Source/Runtime/WebBrowser/Private/SWebBrowserView.cpp b/Engine/Source/Runtime/WebBrowser/Private/SWebBrowserView.cpp
--- a/Engine/Source/Runtime/WebBrowser/Private/SWebBrowserView.cpp
+++ b/Engine/Source/Runtime/WebBrowser/Private/SWebBrowserView.cpp
@@ -452,7 +452,7 @@ void SWebBrowserView::HandleBrowserWindowDocumentStateChanged(EWebBrowserDocumen
 
 void SWebBrowserView::HandleBrowserWindowNeedsRedraw()
 {
-	if (FSlateApplication::Get().IsSlateAsleep())
+	if (FSlateApplication::IsInitialized() && FSlateApplication::Get().IsSlateAsleep())
 	{
 		// Tell slate that the widget needs to wake up for one frame to get redrawn
 		RegisterActiveTimer(0.f, FWidgetActiveTimerDelegate::CreateLambda([this](double InCurrentTime, float InDeltaTime) { return EActiveTimerReturnType::Stop; }));
