WebBrowser - Fix crash on shutdown if a CEFWebBrowserWindow was open

diff --git a/Engine/Source/Runtime/WebBrowser/Private/CEF/CEFWebBrowserWindow.cpp b/Engine/Source/Runtime/WebBrowser/Private/CEF/CEFWebBrowserWindow.cpp
--- a/Engine/Source/Runtime/WebBrowser/Private/CEF/CEFWebBrowserWindow.cpp
+++ b/Engine/Source/Runtime/WebBrowser/Private/CEF/CEFWebBrowserWindow.cpp
@@ -452,7 +452,8 @@ void FCEFWebBrowserWindow::ReleaseTextures()
 		if (UpdatableTextures[I] != nullptr)
 		{
 			FSlateUpdatableTexture* TextureToRelease = UpdatableTextures[I];
-			AsyncTask(ENamedThreads::GameThread, [TextureToRelease]()
+
+			if (IsInGameThread())
 			{
 				if (FSlateApplication::IsInitialized())
 				{
@@ -461,7 +462,20 @@ void FCEFWebBrowserWindow::ReleaseTextures()
 						Renderer->ReleaseUpdatableTexture(TextureToRelease);
 					}
 				}
-			});
+			}
+			else if (FTaskGraphInterface::IsRunning())
+			{
+				AsyncTask(ENamedThreads::GameThread, [TextureToRelease]()
+				{
+					if (FSlateApplication::IsInitialized())
+					{
+						if (FSlateRenderer* Renderer = FSlateApplication::Get().GetRenderer())
+						{
+							Renderer->ReleaseUpdatableTexture(TextureToRelease);
+						}
+					}
+				});
+			}
 
 			UpdatableTextures[I] = nullptr;
 		}
