- Fixed crash on Mac in CEF with pressing keys in Web Views (like Quixel Bridge) that are bound to the main menu.

diff --git a/Engine/Source/Runtime/WebBrowser/Private/CEF/CEFWebBrowserWindow.cpp b/Engine/Source/Runtime/WebBrowser/Private/CEF/CEFWebBrowserWindow.cpp
--- a/Engine/Source/Runtime/WebBrowser/Private/CEF/CEFWebBrowserWindow.cpp
+++ b/Engine/Source/Runtime/WebBrowser/Private/CEF/CEFWebBrowserWindow.cpp
@@ -1365,6 +1365,13 @@ bool FCEFWebBrowserWindow::OnUnhandledKeyEvent(const CefKeyEvent& CefEvent)
 			break;
 		}
 	}
+
+#if PLATFORM_MAC
+	// if this returns false, then CEF will bubble up the keyevent to eventually call NSMenu functions that are not safe on the
+	// game thread, so by returning true, we stop the key handling in it's tracks before CEF can pass it to NSMenu
+	return true;
+#endif
+	
 	return bWasHandled;
 }
 
