Disable Android WebBrowser window from passing on touch and keyboard events for non-3D mode

#android

diff --git a/Engine/Source/Runtime/WebBrowser/Private/Android/AndroidWebBrowserWidget.cpp b/Engine/Source/Runtime/WebBrowser/Private/Android/AndroidWebBrowserWidget.cpp
--- a/Engine/Source/Runtime/WebBrowser/Private/Android/AndroidWebBrowserWidget.cpp
+++ b/Engine/Source/Runtime/WebBrowser/Private/Android/AndroidWebBrowserWidget.cpp
@@ -556,14 +556,17 @@ FReply SAndroidWebBrowserWidget::OnMouseButtonDown(const FGeometry& MyGeometry,
 {
 	FReply Reply = FReply::Unhandled();
 
-	FKey Button = MouseEvent.GetEffectingButton();
-	bool bSupportedButton = (Button == EKeys::LeftMouseButton); // || Button == EKeys::RightMouseButton || Button == EKeys::MiddleMouseButton);
-
-	if (bSupportedButton)
+	if (IsAndroid3DBrowser)
 	{
-		Reply = FReply::Handled();
-		SendTouchDown(ConvertMouseEventToLocal(MyGeometry, MouseEvent));
-		bMouseCapture = true;
+		FKey Button = MouseEvent.GetEffectingButton();
+		bool bSupportedButton = (Button == EKeys::LeftMouseButton); // || Button == EKeys::RightMouseButton || Button == EKeys::MiddleMouseButton);
+
+		if (bSupportedButton)
+		{
+			Reply = FReply::Handled();
+			SendTouchDown(ConvertMouseEventToLocal(MyGeometry, MouseEvent));
+			bMouseCapture = true;
+		}
 	}
 
 	return Reply;
@@ -573,14 +576,17 @@ FReply SAndroidWebBrowserWidget::OnMouseButtonUp(const FGeometry& MyGeometry, co
 {
 	FReply Reply = FReply::Unhandled();
 
-	FKey Button = MouseEvent.GetEffectingButton();
-	bool bSupportedButton = (Button == EKeys::LeftMouseButton); // || Button == EKeys::RightMouseButton || Button == EKeys::MiddleMouseButton);
-
-	if (bSupportedButton)
+	if (IsAndroid3DBrowser)
 	{
-		Reply = FReply::Handled();
-		SendTouchUp(ConvertMouseEventToLocal(MyGeometry, MouseEvent));
-		bMouseCapture = false;
+		FKey Button = MouseEvent.GetEffectingButton();
+		bool bSupportedButton = (Button == EKeys::LeftMouseButton); // || Button == EKeys::RightMouseButton || Button == EKeys::MiddleMouseButton);
+
+		if (bSupportedButton)
+		{
+			Reply = FReply::Handled();
+			SendTouchUp(ConvertMouseEventToLocal(MyGeometry, MouseEvent));
+			bMouseCapture = false;
+		}
 	}
 
 	return Reply;
@@ -590,7 +596,7 @@ FReply SAndroidWebBrowserWidget::OnMouseMove(const FGeometry& MyGeometry, const
 {
 	FReply Reply = FReply::Unhandled();
 
-	if (bMouseCapture)
+	if (IsAndroid3DBrowser && bMouseCapture)
 	{
 		Reply = FReply::Handled();
 		SendTouchMove(ConvertMouseEventToLocal(MyGeometry, MouseEvent));
@@ -601,12 +607,22 @@ FReply SAndroidWebBrowserWidget::OnMouseMove(const FGeometry& MyGeometry, const
 
 FReply SAndroidWebBrowserWidget::OnKeyDown(const FGeometry& MyGeometry, const FKeyEvent& InKeyEvent)
 {
+	if (!IsAndroid3DBrowser)
+	{
+		return FReply::Unhandled();
+	}
+
 //	FPlatformMisc::LowLevelOutputDebugStringf(TEXT("SAndroidWebBrowserWidget::OnKeyDown: %d"), InKeyEvent.GetCharacter());
 	return JavaWebBrowser->SendKeyDown(InKeyEvent.GetCharacter()) ? FReply::Handled() : FReply::Unhandled();
 }
 
 FReply SAndroidWebBrowserWidget::OnKeyUp(const FGeometry& MyGeometry, const FKeyEvent& InKeyEvent)
 {
+	if (!IsAndroid3DBrowser)
+	{
+		return FReply::Unhandled();
+	}
+
 //	FPlatformMisc::LowLevelOutputDebugStringf(TEXT("SAndroidWebBrowserWidget::OnKeyUp: %d"), InKeyEvent.GetCharacter());
 	return JavaWebBrowser->SendKeyUp(InKeyEvent.GetCharacter()) ? FReply::Handled() : FReply::Unhandled();
 }
@@ -614,10 +630,13 @@ FReply SAndroidWebBrowserWidget::OnKeyUp(const FGeometry& MyGeometry, const FKey
 FReply SAndroidWebBrowserWidget::OnKeyChar(const FGeometry& MyGeometry, const FCharacterEvent& InCharacterEvent)
 {
 //	FPlatformMisc::LowLevelOutputDebugStringf(TEXT("SAndroidWebBrowserWidget::OnKeyChar: %d"), (int32)InCharacterEvent.GetCharacter());
-	if (JavaWebBrowser->SendKeyDown(InCharacterEvent.GetCharacter()))
+	if (IsAndroid3DBrowser)
 	{
-		JavaWebBrowser->SendKeyUp(InCharacterEvent.GetCharacter());
-		return FReply::Handled();
+		if (JavaWebBrowser->SendKeyDown(InCharacterEvent.GetCharacter()))
+		{
+			JavaWebBrowser->SendKeyUp(InCharacterEvent.GetCharacter());
+			return FReply::Handled();
+		}
 	}
 	return FReply::Unhandled();
 }
diff --git a/Engine/Source/Runtime/WebBrowser/Private/Android/AndroidWebBrowserWindow.cpp b/Engine/Source/Runtime/WebBrowser/Private/Android/AndroidWebBrowserWindow.cpp
--- a/Engine/Source/Runtime/WebBrowser/Private/Android/AndroidWebBrowserWindow.cpp
+++ b/Engine/Source/Runtime/WebBrowser/Private/Android/AndroidWebBrowserWindow.cpp
@@ -131,17 +131,20 @@ FString FAndroidWebBrowserWindow::GetUrl() const
 
 bool FAndroidWebBrowserWindow::OnKeyDown(const FKeyEvent& InKeyEvent)
 {
-	return BrowserWidget->OnKeyDown(FGeometry(), InKeyEvent).IsEventHandled();
+//	return BrowserWidget->OnKeyDown(FGeometry(), InKeyEvent).IsEventHandled();
+	return false;
 }
 
 bool FAndroidWebBrowserWindow::OnKeyUp(const FKeyEvent& InKeyEvent)
 {
-	return BrowserWidget->OnKeyUp(FGeometry(), InKeyEvent).IsEventHandled();
+//	return BrowserWidget->OnKeyUp(FGeometry(), InKeyEvent).IsEventHandled();
+	return false;
 }
 
 bool FAndroidWebBrowserWindow::OnKeyChar(const FCharacterEvent& InCharacterEvent)
 {
-	return BrowserWidget->OnKeyChar(FGeometry(), InCharacterEvent).IsEventHandled();
+//	return BrowserWidget->OnKeyChar(FGeometry(), InCharacterEvent).IsEventHandled();
+	return false;
 }
 
 FVector2D FAndroidWebBrowserWindow::ConvertMouseEventToLocal(const FGeometry& MyGeometry, const FPointerEvent& MouseEvent, bool bIsPopup)
@@ -157,7 +160,7 @@ FVector2D FAndroidWebBrowserWindow::ConvertMouseEventToLocal(const FGeometry& My
 FReply FAndroidWebBrowserWindow::OnMouseButtonDown(const FGeometry& MyGeometry, const FPointerEvent& MouseEvent, bool bIsPopup)
 {
 	FReply Reply = FReply::Unhandled();
-
+/*
 	FKey Button = MouseEvent.GetEffectingButton();
 	bool bSupportedButton = (Button == EKeys::LeftMouseButton); // || Button == EKeys::RightMouseButton || Button == EKeys::MiddleMouseButton);
 
@@ -166,14 +169,14 @@ FReply FAndroidWebBrowserWindow::OnMouseButtonDown(const FGeometry& MyGeometry,
 		Reply = FReply::Handled();
 		BrowserWidget->SendTouchDown(ConvertMouseEventToLocal(MyGeometry, MouseEvent, bIsPopup));
 	}
-
+*/
 	return Reply;
 }
 
 FReply FAndroidWebBrowserWindow::OnMouseButtonUp(const FGeometry& MyGeometry, const FPointerEvent& MouseEvent, bool bIsPopup)
 {
 	FReply Reply = FReply::Unhandled();
-
+/*
 	FKey Button = MouseEvent.GetEffectingButton();
 	bool bSupportedButton = (Button == EKeys::LeftMouseButton); // || Button == EKeys::RightMouseButton || Button == EKeys::MiddleMouseButton);
 
@@ -182,7 +185,7 @@ FReply FAndroidWebBrowserWindow::OnMouseButtonUp(const FGeometry& MyGeometry, co
 		Reply = FReply::Handled();
 		BrowserWidget->SendTouchUp(ConvertMouseEventToLocal(MyGeometry, MouseEvent, bIsPopup));
 	}
-
+*/
 	return Reply;
 }
 
@@ -194,7 +197,7 @@ FReply FAndroidWebBrowserWindow::OnMouseButtonDoubleClick(const FGeometry& MyGeo
 FReply FAndroidWebBrowserWindow::OnMouseMove(const FGeometry& MyGeometry, const FPointerEvent& MouseEvent, bool bIsPopup)
 {
 	FReply Reply = FReply::Unhandled();
-
+/*
 	FKey Button = MouseEvent.GetEffectingButton();
 	bool bSupportedButton = (Button == EKeys::LeftMouseButton); // || Button == EKeys::RightMouseButton || Button == EKeys::MiddleMouseButton);
 
@@ -203,7 +206,7 @@ FReply FAndroidWebBrowserWindow::OnMouseMove(const FGeometry& MyGeometry, const
 		Reply = FReply::Handled();
 		BrowserWidget->SendTouchMove(ConvertMouseEventToLocal(MyGeometry, MouseEvent, bIsPopup));
 	}
-
+*/
 	return Reply;
 }
 
