[Engine]

diff --git a/Engine/Source/Runtime/UMG/Public/Slate/SRetainerWidget.h b/Engine/Source/Runtime/UMG/Public/Slate/SRetainerWidget.h
--- a/Engine/Source/Runtime/UMG/Public/Slate/SRetainerWidget.h
+++ b/Engine/Source/Runtime/UMG/Public/Slate/SRetainerWidget.h
@@ -36,10 +36,10 @@ DECLARE_MULTICAST_DELEGATE( FOnRetainedModeChanged );
  * frequency of the main game render.  It also has the side benefit of allow materials
  * to be applied to the render target after drawing the widgets to apply a simple post process.
  */
-class UMG_API SRetainerWidget : public SCompoundWidget,  public FSlateInvalidationRoot
+class SRetainerWidget : public SCompoundWidget,  public FSlateInvalidationRoot
 {
 public:
-	static int32 Shared_MaxRetainerWorkPerFrame;
+	static UMG_API int32 Shared_MaxRetainerWorkPerFrame;
 
 public:
 	SLATE_BEGIN_ARGS(SRetainerWidget)
@@ -60,51 +60,51 @@ public:
 		SLATE_ARGUMENT(FName, StatId)
 	SLATE_END_ARGS()
 
-	SRetainerWidget();
-	~SRetainerWidget();
+	UMG_API SRetainerWidget();
+	UMG_API ~SRetainerWidget();
 
 	/** Constructor */
-	void Construct(const FArguments& Args);
+	UMG_API void Construct(const FArguments& Args);
 
-	void SetRenderingPhase(int32 Phase, int32 PhaseCount);
+	UMG_API void SetRenderingPhase(int32 Phase, int32 PhaseCount);
 
 	/** Requests that the retainer redraw the hosted content next time it's painted. */
-	void RequestRender();
+	UMG_API void RequestRender();
 
 #if WITH_EDITOR
-	void SetIsDesignTime(bool bInIsDesignTime);
-	void SetShowEffectsInDesigner(bool bInShowEffectsInDesigner);
+	UMG_API void SetIsDesignTime(bool bInIsDesignTime);
+	UMG_API void SetShowEffectsInDesigner(bool bInShowEffectsInDesigner);
 #endif
 
-	void SetRetainedRendering(bool bRetainRendering);
+	UMG_API void SetRetainedRendering(bool bRetainRendering);
 
-	void SetContent(const TSharedRef< SWidget >& InContent);
+	UMG_API void SetContent(const TSharedRef< SWidget >& InContent);
 
-	UMaterialInstanceDynamic* GetEffectMaterial() const;
+	UMG_API UMaterialInstanceDynamic* GetEffectMaterial() const;
 
-	void SetEffectMaterial(UMaterialInterface* EffectMaterial);
+	UMG_API void SetEffectMaterial(UMaterialInterface* EffectMaterial);
 
-	void SetTextureParameter(FName TextureParameter);
+	UMG_API void SetTextureParameter(FName TextureParameter);
 
 	//~ SWidget interface
-	virtual FChildren* GetChildren() override;
+	UMG_API virtual FChildren* GetChildren() override;
 #if WITH_SLATE_DEBUGGING
-	virtual FChildren* Debug_GetChildrenForReflector() override;
+	UMG_API virtual FChildren* Debug_GetChildrenForReflector() override;
 #endif
 
-	void SetWorld(UWorld* World);
+	UMG_API void SetWorld(UWorld* World);
 
 protected:
 	/** SCompoundWidget interface */
-	virtual int32 OnPaint(const FPaintArgs& Args, const FGeometry& AllottedGeometry, const FSlateRect& MyCullingRect, FSlateWindowElementList& OutDrawElements, int32 LayerId, const FWidgetStyle& InWidgetStyle, bool bParentEnabled) const override;
-	virtual FVector2D ComputeDesiredSize(float Scale) const override;
+	UMG_API virtual int32 OnPaint(const FPaintArgs& Args, const FGeometry& AllottedGeometry, const FSlateRect& MyCullingRect, FSlateWindowElementList& OutDrawElements, int32 LayerId, const FWidgetStyle& InWidgetStyle, bool bParentEnabled) const override;
+	UMG_API virtual FVector2D ComputeDesiredSize(float Scale) const override;
 	virtual bool Advanced_IsInvalidationRoot() const { return bEnableRetainedRendering; }
 	virtual const FSlateInvalidationRoot* Advanced_AsInvalidationRoot() const override { return bEnableRetainedRendering ? this : nullptr; }
-	virtual bool CustomPrepass(float LayoutScaleMultiplier) override;
+	UMG_API virtual bool CustomPrepass(float LayoutScaleMultiplier) override;
 
 	//~ Begin FSlateInvalidationRoot interface
-	virtual TSharedRef<SWidget> GetRootWidget() override;
-	virtual int32 PaintSlowPath(const FSlateInvalidationContext& Context) override;
+	UMG_API virtual TSharedRef<SWidget> GetRootWidget() override;
+	UMG_API virtual int32 PaintSlowPath(const FSlateInvalidationContext& Context) override;
 
 	enum class EPaintRetainedContentResult
 	{
@@ -114,20 +114,20 @@ protected:
 		TextureSizeTooBig,
 		TextureSizeZero,
 	};
-	EPaintRetainedContentResult PaintRetainedContentImpl(const FSlateInvalidationContext& Context, const FGeometry& AllottedGeometry, int32 LayerId);
+	UMG_API EPaintRetainedContentResult PaintRetainedContentImpl(const FSlateInvalidationContext& Context, const FGeometry& AllottedGeometry, int32 LayerId);
 	//~ End FSlateInvalidationRoot interface
 
-	void RefreshRenderingMode();
-	bool ShouldBeRenderingOffscreen() const;
-	bool IsAnythingVisibleToRender() const;
-	void OnRetainerModeChanged();
-	void OnRootInvalidated();
+	UMG_API void RefreshRenderingMode();
+	UMG_API bool ShouldBeRenderingOffscreen() const;
+	UMG_API bool IsAnythingVisibleToRender() const;
+	UMG_API void OnRetainerModeChanged();
+	UMG_API void OnRootInvalidated();
 
 private:
-	void OnGlobalInvalidationToggled(bool bGlobalInvalidationEnabled);
+	UMG_API void OnGlobalInvalidationToggled(bool bGlobalInvalidationEnabled);
 #if !UE_BUILD_SHIPPING
-	static void OnRetainerModeCVarChanged( IConsoleVariable* CVar );
-	static FOnRetainedModeChanged OnRetainerModeChangedDelegate;
+	static UMG_API void OnRetainerModeCVarChanged( IConsoleVariable* CVar );
+	static UMG_API FOnRetainedModeChanged OnRetainerModeChangedDelegate;
 #endif
 
 	mutable FSlateBrush SurfaceBrush;
@@ -139,7 +139,7 @@ private:
 	FColor PreviousColorAndOpacity;
 	int32 LastIncomingLayerId;
 
-	void UpdateWidgetRenderer();
+	UMG_API void UpdateWidgetRenderer();
 
 	TSharedPtr<SWidget> MyWidget;
 	TSharedRef<SVirtualWindow> VirtualWindow;
@@ -180,6 +180,6 @@ private:
 
 	FName DynamicEffectTextureParameter;
 
-	static TArray<SRetainerWidget*, TInlineAllocator<3>> Shared_WaitingToRender;
-	static TFrameValue<int32> Shared_RetainerWorkThisFrame;
+	static UMG_API TArray<SRetainerWidget*, TInlineAllocator<3>> Shared_WaitingToRender;
+	static UMG_API TFrameValue<int32> Shared_RetainerWorkThisFrame;
 };
diff --git a/Engine/Source/Runtime/UMG/Public/Slate/SlateVectorArtData.h b/Engine/Source/Runtime/UMG/Public/Slate/SlateVectorArtData.h
--- a/Engine/Source/Runtime/UMG/Public/Slate/SlateVectorArtData.h
+++ b/Engine/Source/Runtime/UMG/Public/Slate/SlateVectorArtData.h
@@ -80,39 +80,39 @@ struct FSlateMeshVertex
 /**
  * Turn static mesh data into Slate's simple vector art format.
  */
-UCLASS()
-class UMG_API USlateVectorArtData : public UObject
+UCLASS(MinimalAPI)
+class USlateVectorArtData : public UObject
 {
 	GENERATED_UCLASS_BODY()
 
 public:
 	/** Access the slate vertexes. */
-	const TArray<FSlateMeshVertex>& GetVertexData() const;
+	UMG_API const TArray<FSlateMeshVertex>& GetVertexData() const;
 	
 	/** Access the indexes for the order in which to draw the vertexes. */
-	const TArray<uint32>& GetIndexData() const;
+	UMG_API const TArray<uint32>& GetIndexData() const;
 	
 	/** Material to be used with the specified vector art data. */
-	UMaterialInterface* GetMaterial() const;
+	UMG_API UMaterialInterface* GetMaterial() const;
 	
 	/** Convert the material into an MID and get a pointer to the MID so that parameters can be set on it. */
-	UMaterialInstanceDynamic* ConvertToMaterialInstanceDynamic();
+	UMG_API UMaterialInstanceDynamic* ConvertToMaterialInstanceDynamic();
 
 	/** Convert the static mesh data into slate vector art on demand. Does nothing in a cooked build. */
-	void EnsureValidData();
+	UMG_API void EnsureValidData();
 
-	FVector2D GetDesiredSize() const;
+	UMG_API FVector2D GetDesiredSize() const;
 
-	FVector2D GetExtentMin() const;
+	UMG_API FVector2D GetExtentMin() const;
 
-	FVector2D GetExtentMax() const;
+	UMG_API FVector2D GetExtentMax() const;
 
 private:
 	// ~ UObject Interface
 	PRAGMA_DISABLE_DEPRECATION_WARNINGS // Suppress compiler warning on override of deprecated function
 	UE_DEPRECATED(5.0, "Use version that takes FObjectPreSaveContext instead.")
-	virtual void PreSave(const class ITargetPlatform* TargetPlatform) override;
-	PRAGMA_ENABLE_DEPRECATION_WARNINGS
+	UMG_API virtual void PreSave(const class ITargetPlatform* TargetPlatform) override;
+	UMG_API PRAGMA_ENABLE_DEPRECATION_WARNINGS
 	virtual void PreSave(FObjectPreSaveContext ObjectSaveContext) override;
 	// ~ UObject Interface
 
diff --git a/Engine/Source/Runtime/UMG/Public/Slate/SlateVectorArtInstanceData.h b/Engine/Source/Runtime/UMG/Public/Slate/SlateVectorArtInstanceData.h
--- a/Engine/Source/Runtime/UMG/Public/Slate/SlateVectorArtInstanceData.h
+++ b/Engine/Source/Runtime/UMG/Public/Slate/SlateVectorArtInstanceData.h
@@ -8,18 +8,18 @@
 #include "Math/Vector4.h"
 #include "Misc/AssertionMacros.h"
 
-struct UMG_API FSlateVectorArtInstanceData
+struct FSlateVectorArtInstanceData
 {
 public:
 	const FVector4& GetData() const { return Data; }
 	FVector4& GetData() { return Data; }
 
-	void SetPositionFixedPoint16(FVector2D Position);
-	void SetScaleFixedPoint16(float Scale);
+	UMG_API void SetPositionFixedPoint16(FVector2D Position);
+	UMG_API void SetScaleFixedPoint16(float Scale);
 
-	void SetPosition(FVector2D Position);
-	void SetScale(float Scale);
-	void SetBaseAddress(float Address);
+	UMG_API void SetPosition(FVector2D Position);
+	UMG_API void SetScale(float Scale);
+	UMG_API void SetBaseAddress(float Address);
 
 	template<int32 Component, int32 ByteIndex>
 	void PackFloatIntoByte(float InValue)
diff --git a/Engine/Source/Runtime/UMG/Public/Slate/WidgetRenderer.h b/Engine/Source/Runtime/UMG/Public/Slate/WidgetRenderer.h
--- a/Engine/Source/Runtime/UMG/Public/Slate/WidgetRenderer.h
+++ b/Engine/Source/Runtime/UMG/Public/Slate/WidgetRenderer.h
@@ -28,11 +28,11 @@ class FRenderTarget;
 /**
  * 
  */
-class UMG_API FWidgetRenderer : public FDeferredCleanupInterface
+class FWidgetRenderer : public FDeferredCleanupInterface
 {
 public:
-	FWidgetRenderer(bool bUseGammaCorrection = false, bool bInClearTarget = true);
-	~FWidgetRenderer();
+	UMG_API FWidgetRenderer(bool bUseGammaCorrection = false, bool bInClearTarget = true);
+	UMG_API ~FWidgetRenderer();
 
 	bool GetIsPrepassNeeded() const { return bPrepassNeeded; }
 	void SetIsPrepassNeeded(bool bInPrepassNeeded) { bPrepassNeeded = bInPrepassNeeded; }
@@ -43,31 +43,31 @@ public:
 	void SetShouldClearTarget(bool bShouldClear) { bClearTarget = bShouldClear; }
 
 	bool GetUseGammaCorrection() const { return bUseGammaSpace; }
-	void SetUseGammaCorrection(bool bInUseGammaSpace);
+	UMG_API void SetUseGammaCorrection(bool bInUseGammaSpace);
 
-	void SetApplyColorDeficiencyCorrection(bool bInApplyColorCorrection);
+	UMG_API void SetApplyColorDeficiencyCorrection(bool bInApplyColorCorrection);
 
-	ISlate3DRenderer* GetSlateRenderer();
+	UMG_API ISlate3DRenderer* GetSlateRenderer();
 
-	static UTextureRenderTarget2D* CreateTargetFor(FVector2D DrawSize, TextureFilter InFilter, bool bUseGammaCorrection);
+	static UMG_API UTextureRenderTarget2D* CreateTargetFor(FVector2D DrawSize, TextureFilter InFilter, bool bUseGammaCorrection);
 
-	UTextureRenderTarget2D* DrawWidget(const TSharedRef<SWidget>& Widget, FVector2D DrawSize);
+	UMG_API UTextureRenderTarget2D* DrawWidget(const TSharedRef<SWidget>& Widget, FVector2D DrawSize);
 
-	void DrawWidget(
+	UMG_API void DrawWidget(
 		FRenderTarget* RenderTarget,
 		const TSharedRef<SWidget>& Widget,
 		FVector2D DrawSize,
 		float DeltaTime,
 		bool bDeferRenderTargetUpdate = false);
 
-	void DrawWidget(
+	UMG_API void DrawWidget(
 		UTextureRenderTarget2D* RenderTarget,
 		const TSharedRef<SWidget>& Widget,
 		FVector2D DrawSize,
 		float DeltaTime,
 		bool bDeferRenderTargetUpdate = false);
 
-	void DrawWidget(
+	UMG_API void DrawWidget(
 		FRenderTarget* RenderTarget,
 		const TSharedRef<SWidget>& Widget,
 		float Scale,
@@ -75,7 +75,7 @@ public:
 		float DeltaTime,
 		bool bDeferRenderTargetUpdate = false);
 
-	void DrawWidget(
+	UMG_API void DrawWidget(
 		UTextureRenderTarget2D* RenderTarget,
 		const TSharedRef<SWidget>& Widget,
 		float Scale,
@@ -83,7 +83,7 @@ public:
 		float DeltaTime,
 		bool bDeferRenderTargetUpdate = false);
 
-	void DrawWindow(
+	UMG_API void DrawWindow(
 		FRenderTarget* RenderTarget,
 		FHittestGrid& HitTestGrid,
 		TSharedRef<SWindow> Window,
@@ -92,7 +92,7 @@ public:
 		float DeltaTime,
 		bool bDeferRenderTargetUpdate = false);
 
-	void DrawWindow(
+	UMG_API void DrawWindow(
 		UTextureRenderTarget2D* RenderTarget,
 		FHittestGrid& HitTestGrid,
 		TSharedRef<SWindow> Window,
@@ -101,7 +101,7 @@ public:
 		float DeltaTime,
 		bool bDeferRenderTargetUpdate = false);
 
-	void DrawWindow(
+	UMG_API void DrawWindow(
 		FRenderTarget* RenderTarget,
 		FHittestGrid& HitTestGrid,
 		TSharedRef<SWindow> Window,
@@ -110,7 +110,7 @@ public:
 		float DeltaTime,
 		bool bDeferRenderTargetUpdate = false);
 
-	void DrawWindow(
+	UMG_API void DrawWindow(
 		UTextureRenderTarget2D* RenderTarget,
 		FHittestGrid& HitTestGrid,
 		TSharedRef<SWindow> Window,
@@ -119,7 +119,7 @@ public:
 		float DeltaTime,
 		bool bDeferRenderTargetUpdate = false);
 
-	void DrawWindow(
+	UMG_API void DrawWindow(
 		const FPaintArgs& PaintArgs,
 		FRenderTarget* RenderTarget,
 		TSharedRef<SWindow> Window,
@@ -128,7 +128,7 @@ public:
 		float DeltaTime,
 		bool bDeferRenderTargetUpdate = false);
 
-	void DrawWindow(
+	UMG_API void DrawWindow(
 		const FPaintArgs& PaintArgs,
 		UTextureRenderTarget2D* RenderTarget,
 		TSharedRef<SWindow> Window,
@@ -137,7 +137,7 @@ public:
 		float DeltaTime,
 		bool bDeferRenderTargetUpdate = false);
 
-	bool DrawInvalidationRoot(
+	UMG_API bool DrawInvalidationRoot(
 		TSharedRef<SVirtualWindow>& VirtualWindow,
 		UTextureRenderTarget2D* RenderTarget,
 		FSlateInvalidationRoot& Root,
diff --git a/Engine/Source/Runtime/UMG/Public/UMGStyle.h b/Engine/Source/Runtime/UMG/Public/UMGStyle.h
--- a/Engine/Source/Runtime/UMG/Public/UMGStyle.h
+++ b/Engine/Source/Runtime/UMG/Public/UMGStyle.h
@@ -6,21 +6,21 @@
 #include "Styling/SlateStyle.h"
 
 /**  */
-class UMG_API FUMGStyle
+class FUMGStyle
 {
 public:
 
-	static void Initialize();
+	static UMG_API void Initialize();
 
-	static void Shutdown();
+	static UMG_API void Shutdown();
 
 	/** reloads textures used by slate renderer */
-	static void ReloadTextures();
+	static UMG_API void ReloadTextures();
 
 	/** @return The Slate style set for the UMG Style */
-	static const ISlateStyle& Get();
+	static UMG_API const ISlateStyle& Get();
 
-	static FName GetStyleSetName();
+	static UMG_API FName GetStyleSetName();
 
 private:
 
diff --git a/Engine/Source/Runtime/WebBrowser/Public/IWebBrowserAdapter.h b/Engine/Source/Runtime/WebBrowser/Public/IWebBrowserAdapter.h
--- a/Engine/Source/Runtime/WebBrowser/Public/IWebBrowserAdapter.h
+++ b/Engine/Source/Runtime/WebBrowser/Public/IWebBrowserAdapter.h
@@ -18,11 +18,11 @@ public:
 
 };
 
-class WEBBROWSER_API FWebBrowserAdapterFactory 
+class FWebBrowserAdapterFactory 
 { 
 public: 
 
-	static TSharedRef<IWebBrowserAdapter> Create(const FString& Name, UObject* JSBridge, bool IsPermanent); 
+	static WEBBROWSER_API TSharedRef<IWebBrowserAdapter> Create(const FString& Name, UObject* JSBridge, bool IsPermanent); 
 
-	static TSharedRef<IWebBrowserAdapter> Create(const FString& Name, UObject* JSBridge, bool IsPermanent, const FString& ConnectScriptText, const FString& DisconnectScriptText);
+	static WEBBROWSER_API TSharedRef<IWebBrowserAdapter> Create(const FString& Name, UObject* JSBridge, bool IsPermanent, const FString& ConnectScriptText, const FString& DisconnectScriptText);
 }; 
diff --git a/Engine/Source/Runtime/WebBrowser/Public/IWebBrowserSchemeHandler.h b/Engine/Source/Runtime/WebBrowser/Public/IWebBrowserSchemeHandler.h
--- a/Engine/Source/Runtime/WebBrowser/Public/IWebBrowserSchemeHandler.h
+++ b/Engine/Source/Runtime/WebBrowser/Public/IWebBrowserSchemeHandler.h
@@ -7,7 +7,7 @@
  * This is the interface that needs to be implemented to handle a request made via a custom scheme.
  * It will be created by implementing an IWebBrowserSchemeHandlerFactory, given to the web browser singleton.
  */
-class WEBBROWSER_API IWebBrowserSchemeHandler
+class IWebBrowserSchemeHandler
 {
 public:
 	/**
@@ -88,7 +88,7 @@ public:
 /**
  * This is the interface that needs to be implemented to instantiate a scheme request handler.
  */
-class WEBBROWSER_API IWebBrowserSchemeHandlerFactory
+class IWebBrowserSchemeHandlerFactory
 {
 public:
 	virtual ~IWebBrowserSchemeHandlerFactory() {}
diff --git a/Engine/Source/Runtime/WebBrowser/Public/IWebBrowserSingleton.h b/Engine/Source/Runtime/WebBrowser/Public/IWebBrowserSingleton.h
--- a/Engine/Source/Runtime/WebBrowser/Public/IWebBrowserSingleton.h
+++ b/Engine/Source/Runtime/WebBrowser/Public/IWebBrowserSingleton.h
@@ -31,7 +31,7 @@ public:
 		FColor BackgroundColor = FColor(255, 255, 255, 255)) = 0;
 };
 
-struct WEBBROWSER_API FBrowserContextSettings
+struct FBrowserContextSettings
 {
 	FBrowserContextSettings(const FString& InId)
 		: Id(InId)
@@ -52,7 +52,7 @@ struct WEBBROWSER_API FBrowserContextSettings
 };
 
 
-struct WEBBROWSER_API FCreateBrowserWindowSettings
+struct FCreateBrowserWindowSettings
 {
 
 	FCreateBrowserWindowSettings()
@@ -85,7 +85,7 @@ struct WEBBROWSER_API FCreateBrowserWindowSettings
 /**
  * A singleton class that takes care of general web browser tasks
  */
-class WEBBROWSER_API IWebBrowserSingleton
+class IWebBrowserSingleton
 {
 public:
 	/**
diff --git a/Engine/Source/Runtime/WebBrowser/Public/SWebBrowser.h b/Engine/Source/Runtime/WebBrowser/Public/SWebBrowser.h
--- a/Engine/Source/Runtime/WebBrowser/Public/SWebBrowser.h
+++ b/Engine/Source/Runtime/WebBrowser/Public/SWebBrowser.h
@@ -18,7 +18,7 @@ enum class EWebBrowserDialogEventResponse;
 
 DECLARE_DELEGATE_RetVal(bool, FOnSuppressContextMenu);
 
-class WEBBROWSER_API SWebBrowser
+class SWebBrowser
 	: public SCompoundWidget
 {
 public:
@@ -128,9 +128,9 @@ public:
 
 
 	/** Default constructor. */
-	SWebBrowser();
+	WEBBROWSER_API SWebBrowser();
 
-	~SWebBrowser();
+	WEBBROWSER_API ~SWebBrowser();
 
 	virtual bool SupportsKeyboardFocus() const override {return true;}
 
@@ -139,14 +139,14 @@ public:
 	 *
 	 * @param InArgs  Declaration from which to construct the widget.
 	 */
-	void Construct(const FArguments& InArgs, const TSharedPtr<IWebBrowserWindow>& InWebBrowserWindow = nullptr);
+	WEBBROWSER_API void Construct(const FArguments& InArgs, const TSharedPtr<IWebBrowserWindow>& InWebBrowserWindow = nullptr);
 
 	/**
 	 * Load the specified URL.
 	 *
 	 * @param NewURL New URL to load.
 	 */
-	void LoadURL(FString NewURL);
+	WEBBROWSER_API void LoadURL(FString NewURL);
 
 	/**
 	* Load a string as data to create a web page.
@@ -154,39 +154,39 @@ public:
 	* @param Contents String to load.
 	* @param DummyURL Dummy URL for the page.
 	*/
-	void LoadString(FString Contents, FString DummyURL);
+	WEBBROWSER_API void LoadString(FString Contents, FString DummyURL);
 
 	/** Reload the current page. */
-	void Reload();
+	WEBBROWSER_API void Reload();
 
 	/** Stop loading the page. */
-	void StopLoad();
+	WEBBROWSER_API void StopLoad();
 
 	/** Get the current title of the web page. */
-	FText GetTitleText() const;
+	WEBBROWSER_API FText GetTitleText() const;
 
 	/**
 	 * Gets the currently loaded URL.
 	 *
 	 * @return The URL, or empty string if no document is loaded.
 	 */
-	FString GetUrl() const;
+	WEBBROWSER_API FString GetUrl() const;
 
 	/**
 	 * Gets the URL that appears in the address bar, this may not be the URL that is currently loaded in the frame.
 	 *
 	 * @return The address bar URL.
 	 */
-	FText GetAddressBarUrlText() const;
+	WEBBROWSER_API FText GetAddressBarUrlText() const;
 
 	/** Whether the document finished loading. */
-	bool IsLoaded() const;
+	WEBBROWSER_API bool IsLoaded() const;
 
 	/** Whether the document is currently being loaded. */
-	bool IsLoading() const;
+	WEBBROWSER_API bool IsLoading() const;
 
 	/** Execute javascript on the current window */
-	void ExecuteJavascript(const FString& ScriptText);
+	WEBBROWSER_API void ExecuteJavascript(const FString& ScriptText);
 
 	/**
 	 * Gets the source of the main frame as raw HTML.
@@ -195,7 +195,7 @@ public:
 	 * result is ready.
 	 * @param	Callback	A callable that takes a single string reference for handling the result.
 	 */
-	void GetSource(TFunction<void (const FString&)> Callback) const;
+	WEBBROWSER_API void GetSource(TFunction<void (const FString&)> Callback) const;
 
 	/**
 	 * Expose a UObject instance to the browser runtime.
@@ -206,7 +206,7 @@ public:
 	 * @param Object The object instance.
 	 * @param bIsPermanent If true, the object will be visible to all pages loaded through this browser widget, otherwise, it will be deleted when navigating away from the current page. Non-permanent bindings should be registered from inside an OnLoadStarted event handler in order to be available before JS code starts loading.
 	 */
-	void BindUObject(const FString& Name, UObject* Object, bool bIsPermanent = true);
+	WEBBROWSER_API void BindUObject(const FString& Name, UObject* Object, bool bIsPermanent = true);
 
 	/**
 	 * Remove an existing script binding registered by BindUObject.
@@ -215,53 +215,53 @@ public:
 	 * @param Object The object will only be removed if it is the same object as the one passed in.
 	 * @param bIsPermanent Must match the bIsPermanent argument passed to BindUObject.
 	 */
-	void UnbindUObject(const FString& Name, UObject* Object, bool bIsPermanent = true);
+	WEBBROWSER_API void UnbindUObject(const FString& Name, UObject* Object, bool bIsPermanent = true);
 
-	void BindAdapter(const TSharedRef<IWebBrowserAdapter>& Adapter);
+	WEBBROWSER_API void BindAdapter(const TSharedRef<IWebBrowserAdapter>& Adapter);
 
-	void UnbindAdapter(const TSharedRef<IWebBrowserAdapter>& Adapter);
+	WEBBROWSER_API void UnbindAdapter(const TSharedRef<IWebBrowserAdapter>& Adapter);
 
-	void BindInputMethodSystem(ITextInputMethodSystem* TextInputMethodSystem);
+	WEBBROWSER_API void BindInputMethodSystem(ITextInputMethodSystem* TextInputMethodSystem);
 
-	void UnbindInputMethodSystem();
+	WEBBROWSER_API void UnbindInputMethodSystem();
 
 	/** Returns true if the browser can navigate backwards. */
-	bool CanGoBack() const;
+	WEBBROWSER_API bool CanGoBack() const;
 
 	/** Navigate backwards. */
-	void GoBack();
+	WEBBROWSER_API void GoBack();
 
 	/** Returns true if the browser can navigate forwards. */
-	bool CanGoForward() const;
+	WEBBROWSER_API bool CanGoForward() const;
 
 	/** Navigate forwards. */
-	void GoForward();
+	WEBBROWSER_API void GoForward();
 
 	/** Set parent SWindow for this browser. */
-	void SetParentWindow(TSharedPtr<SWindow> Window);
+	WEBBROWSER_API void SetParentWindow(TSharedPtr<SWindow> Window);
 
 private:
 
 	/** Navigate backwards. */
-	FReply OnBackClicked();
+	WEBBROWSER_API FReply OnBackClicked();
 
 	/** Navigate forwards. */
-	FReply OnForwardClicked();
+	WEBBROWSER_API FReply OnForwardClicked();
 
 	/** Get text for reload button depending on status */
-	FText GetReloadButtonText() const;
+	WEBBROWSER_API FText GetReloadButtonText() const;
 
 	/** Reload or stop loading */
-	FReply OnReloadClicked();
+	WEBBROWSER_API FReply OnReloadClicked();
 
 	/** Invoked whenever text is committed in the address bar. */
-	void OnUrlTextCommitted( const FText& NewText, ETextCommit::Type CommitType );
+	WEBBROWSER_API void OnUrlTextCommitted( const FText& NewText, ETextCommit::Type CommitType );
 
 	/** Get whether the page viewport should be visible */
-	EVisibility GetViewportVisibility() const;
+	WEBBROWSER_API EVisibility GetViewportVisibility() const;
 
 	/** Get whether loading throbber should be visible */
-	EVisibility GetLoadingThrobberVisibility() const;
+	WEBBROWSER_API EVisibility GetLoadingThrobberVisibility() const;
 
 private:
 
diff --git a/Engine/Source/Runtime/WebBrowser/Public/SWebBrowserView.h b/Engine/Source/Runtime/WebBrowser/Public/SWebBrowserView.h
--- a/Engine/Source/Runtime/WebBrowser/Public/SWebBrowserView.h
+++ b/Engine/Source/Runtime/WebBrowser/Public/SWebBrowserView.h
@@ -33,7 +33,7 @@ typedef SViewport SWebBrowserWidget;
 typedef SWidget SWebBrowserWidget;
 #endif
 
-class WEBBROWSER_API SWebBrowserView
+class SWebBrowserView
 	: public SCompoundWidget
 {
 public:
@@ -161,9 +161,9 @@ public:
 
 
 	/** Default constructor. */
-	SWebBrowserView();
+	WEBBROWSER_API SWebBrowserView();
 
-	~SWebBrowserView();
+	WEBBROWSER_API ~SWebBrowserView();
 
 	virtual bool SupportsKeyboardFocus() const override {return true;}
 
@@ -172,16 +172,16 @@ public:
 	 *
 	 * @param InArgs  Declaration from which to construct the widget.
 	 */
-	void Construct(const FArguments& InArgs, const TSharedPtr<IWebBrowserWindow>& InWebBrowserWindow = nullptr);
+	WEBBROWSER_API void Construct(const FArguments& InArgs, const TSharedPtr<IWebBrowserWindow>& InWebBrowserWindow = nullptr);
 
-	virtual int32 OnPaint(const FPaintArgs& Args, const FGeometry& AllottedGeometry, const FSlateRect& MyCullingRect, FSlateWindowElementList& OutDrawElements, int32 LayerId, const FWidgetStyle& InWidgetStyle, bool bParentEnabled) const override;
+	WEBBROWSER_API virtual int32 OnPaint(const FPaintArgs& Args, const FGeometry& AllottedGeometry, const FSlateRect& MyCullingRect, FSlateWindowElementList& OutDrawElements, int32 LayerId, const FWidgetStyle& InWidgetStyle, bool bParentEnabled) const override;
 
 	/**
 	 * Load the specified URL.
 	 *
 	 * @param NewURL New URL to load.
 	 */
-	void LoadURL(FString NewURL);
+	WEBBROWSER_API void LoadURL(FString NewURL);
 
 	/**
 	* Load a string as data to create a web page.
@@ -189,42 +189,42 @@ public:
 	* @param Contents String to load.
 	* @param DummyURL Dummy URL for the page.
 	*/
-	void LoadString(FString Contents, FString DummyURL);
+	WEBBROWSER_API void LoadString(FString Contents, FString DummyURL);
 
 	/** Reload the current page. */
-	void Reload();
+	WEBBROWSER_API void Reload();
 
 	/** Stop loading the page. */
-	void StopLoad();
+	WEBBROWSER_API void StopLoad();
 
 	/** Get the current title of the web page. */
-	FText GetTitleText() const;
+	WEBBROWSER_API FText GetTitleText() const;
 
 	/**
 	 * Gets the currently loaded URL.
 	 *
 	 * @return The URL, or empty string if no document is loaded.
 	 */
-	FString GetUrl() const;
+	WEBBROWSER_API FString GetUrl() const;
 
 	/**
 	 * Gets the URL that appears in the address bar, this may not be the URL that is currently loaded in the frame.
 	 *
 	 * @return The address bar URL.
 	 */
-	FText GetAddressBarUrlText() const;
+	WEBBROWSER_API FText GetAddressBarUrlText() const;
 
 	/** Whether the document finished loading. */
-	bool IsLoaded() const;
+	WEBBROWSER_API bool IsLoaded() const;
 
 	/** Whether the document is currently being loaded. */
-	bool IsLoading() const;
+	WEBBROWSER_API bool IsLoading() const;
 
 	/** Whether the browser widget is done initializing. */
-	bool IsInitialized() const;
+	WEBBROWSER_API bool IsInitialized() const;
 
 	/** Execute javascript on the current window */
-	void ExecuteJavascript(const FString& ScriptText);
+	WEBBROWSER_API void ExecuteJavascript(const FString& ScriptText);
 
 	/**
 	 * Gets the source of the main frame as raw HTML.
@@ -233,7 +233,7 @@ public:
 	 * result is ready.
 	 * @param	Callback	A callable that takes a single string reference for handling the result.
 	 */
-	void GetSource(TFunction<void (const FString&)> Callback) const ;
+	WEBBROWSER_API void GetSource(TFunction<void (const FString&)> Callback) const ;
 
 	/**
 	 * Expose a UObject instance to the browser runtime.
@@ -244,7 +244,7 @@ public:
 	 * @param Object The object instance.
 	 * @param bIsPermanent If true, the object will be visible to all pages loaded through this browser widget, otherwise, it will be deleted when navigating away from the current page. Non-permanent bindings should be registered from inside an OnLoadStarted event handler in order to be available before JS code starts loading.
 	 */
-	void BindUObject(const FString& Name, UObject* Object, bool bIsPermanent = true);
+	WEBBROWSER_API void BindUObject(const FString& Name, UObject* Object, bool bIsPermanent = true);
 
 	/**
 	 * Remove an existing script binding registered by BindUObject.
@@ -253,99 +253,99 @@ public:
 	 * @param Object The object will only be removed if it is the same object as the one passed in.
 	 * @param bIsPermanent Must match the bIsPermanent argument passed to BindUObject.
 	 */
-	void UnbindUObject(const FString& Name, UObject* Object, bool bIsPermanent = true);
+	WEBBROWSER_API void UnbindUObject(const FString& Name, UObject* Object, bool bIsPermanent = true);
 
-	void BindAdapter(const TSharedRef<IWebBrowserAdapter>& Adapter);
+	WEBBROWSER_API void BindAdapter(const TSharedRef<IWebBrowserAdapter>& Adapter);
 
-	void UnbindAdapter(const TSharedRef<IWebBrowserAdapter>& Adapter);
+	WEBBROWSER_API void UnbindAdapter(const TSharedRef<IWebBrowserAdapter>& Adapter);
 
-	void BindInputMethodSystem(ITextInputMethodSystem* TextInputMethodSystem);
+	WEBBROWSER_API void BindInputMethodSystem(ITextInputMethodSystem* TextInputMethodSystem);
 
-	void UnbindInputMethodSystem();
+	WEBBROWSER_API void UnbindInputMethodSystem();
 
 	/** Returns true if the browser can navigate backwards. */
-	bool CanGoBack() const;
+	WEBBROWSER_API bool CanGoBack() const;
 
 	/** Navigate backwards. */
-	void GoBack();
+	WEBBROWSER_API void GoBack();
 
 	/** Returns true if the browser can navigate forwards. */
-	bool CanGoForward() const;
+	WEBBROWSER_API bool CanGoForward() const;
 
 	/** Navigate forwards. */
-	void GoForward();
+	WEBBROWSER_API void GoForward();
 
 	/** Set parent SWindow for this browser. */
-	void SetParentWindow(TSharedPtr<SWindow> Window);
+	WEBBROWSER_API void SetParentWindow(TSharedPtr<SWindow> Window);
 
 	/** Update the underlying browser widget to match the KB focus in slate.
 		This is used to work around a CEF bug that loses focus state on navigations*/
-	void SetBrowserKeyboardFocus();
+	WEBBROWSER_API void SetBrowserKeyboardFocus();
 
 	/** Close the underlying browser object before we destruct this view. 
 	    This will block until that object is fully destroyed.
 		Calling this is optional, CEF has object lifetime requirements that mean on shutdown you must destroy browsers before exit.*/
-	void CloseBrowser();
+	WEBBROWSER_API void CloseBrowser();
 
 private:
 
-	void SetupParentWindowHandlers();
+	WEBBROWSER_API void SetupParentWindowHandlers();
 
 	/** Callback for document loading state changes. */
-	void HandleBrowserWindowDocumentStateChanged(EWebBrowserDocumentState NewState);
+	WEBBROWSER_API void HandleBrowserWindowDocumentStateChanged(EWebBrowserDocumentState NewState);
 
 	/** Callback to tell slate we want to update the contents of the web view based on changes inside the view. */
-	void HandleBrowserWindowNeedsRedraw();
+	WEBBROWSER_API void HandleBrowserWindowNeedsRedraw();
 
 	/** Callback for document title changes. */
-	void HandleTitleChanged(FString NewTitle);
+	WEBBROWSER_API void HandleTitleChanged(FString NewTitle);
 
 	/** Callback for loaded url changes. */
-	void HandleUrlChanged(FString NewUrl);
+	WEBBROWSER_API void HandleUrlChanged(FString NewUrl);
 
 	/** Callback for showing browser tool tips. */
-	void HandleToolTip(FString ToolTipText);
+	WEBBROWSER_API void HandleToolTip(FString ToolTipText);
 
 	/**
 	 * A delegate that is executed prior to browser navigation.
 	 *
 	 * @return true if the navigation was handled an no further action should be taken by the browser, false if the browser should handle.
 	 */
-	bool HandleBeforeNavigation(const FString& Url, const FWebNavigationRequest& Request);
+	WEBBROWSER_API bool HandleBeforeNavigation(const FString& Url, const FWebNavigationRequest& Request);
 
-	bool HandleLoadUrl(const FString& Method, const FString& Url, FString& OutResponse);
+	WEBBROWSER_API bool HandleLoadUrl(const FString& Method, const FString& Url, FString& OutResponse);
 
 	/**
 	 * A delegate that is executed when the browser requests window creation.
 	 *
 	 * @return true if if the window request was handled, false if the browser requesting the new window should be closed.
 	 */
-	bool HandleCreateWindow(const TWeakPtr<IWebBrowserWindow>& NewBrowserWindow, const TWeakPtr<IWebBrowserPopupFeatures>& PopupFeatures);
+	WEBBROWSER_API bool HandleCreateWindow(const TWeakPtr<IWebBrowserWindow>& NewBrowserWindow, const TWeakPtr<IWebBrowserPopupFeatures>& PopupFeatures);
 
 	/**
 	 * A delegate that is executed when closing the browser window.
 	 *
 	 * @return true if if the window close was handled, false otherwise.
 	 */
-	bool HandleCloseWindow(const TWeakPtr<IWebBrowserWindow>& BrowserWindow);
+	WEBBROWSER_API bool HandleCloseWindow(const TWeakPtr<IWebBrowserWindow>& BrowserWindow);
 
 	/** Callback for showing dialogs to the user */
-	EWebBrowserDialogEventResponse HandleShowDialog(const TWeakPtr<IWebBrowserDialog>& DialogParams);
+	WEBBROWSER_API EWebBrowserDialogEventResponse HandleShowDialog(const TWeakPtr<IWebBrowserDialog>& DialogParams);
 
 	/** Callback for dismissing any dialogs previously shown  */
-	void HandleDismissAllDialogs();
+	WEBBROWSER_API void HandleDismissAllDialogs();
 
 	/** Callback for popup window permission */
-	bool HandleBeforePopup(FString URL, FString Target);
+	WEBBROWSER_API bool HandleBeforePopup(FString URL, FString Target);
 
 	/** Callback for showing a popup menu */
-	void HandleShowPopup(const FIntRect& PopupSize);
+	WEBBROWSER_API void HandleShowPopup(const FIntRect& PopupSize);
 
 	/** Callback for hiding the popup menu */
-	void HandleDismissPopup();
+	WEBBROWSER_API void HandleDismissPopup();
 
 	/** Callback from the popup menu notifiying it has been dismissed */
-	void HandleMenuDismissed(TSharedRef<IMenu>);
+	WEBBROWSER_API void HandleMenuDismissed(TSharedRef<IMenu>);
 
 	virtual FPopupMethodReply OnQueryPopupMethod() const override
 	{
@@ -354,16 +354,16 @@ private:
 			: FPopupMethodReply::Unhandled();
 	}
 
-	void HandleWindowDeactivated();
-	void HandleWindowActivated();
-	bool UnhandledKeyDown(const FKeyEvent& KeyEvent);
-	bool UnhandledKeyUp(const FKeyEvent& KeyEvent);
-	bool UnhandledKeyChar(const FCharacterEvent& CharacterEvent);
+	WEBBROWSER_API void HandleWindowDeactivated();
+	WEBBROWSER_API void HandleWindowActivated();
+	WEBBROWSER_API bool UnhandledKeyDown(const FKeyEvent& KeyEvent);
+	WEBBROWSER_API bool UnhandledKeyUp(const FKeyEvent& KeyEvent);
+	WEBBROWSER_API bool UnhandledKeyChar(const FCharacterEvent& CharacterEvent);
 
-	bool HandleDrag(const FPointerEvent& MouseEvent);
-	void HandleConsoleMessage(const FString& Message, const FString& Source, int32 Line, EWebBrowserConsoleLogSeverity Serverity);
+	WEBBROWSER_API bool HandleDrag(const FPointerEvent& MouseEvent);
+	WEBBROWSER_API void HandleConsoleMessage(const FString& Message, const FString& Source, int32 Line, EWebBrowserConsoleLogSeverity Serverity);
 
-	TOptional<FSlateRenderTransform> GetPopupRenderTransform() const;
+	WEBBROWSER_API TOptional<FSlateRenderTransform> GetPopupRenderTransform() const;
 private:
 
 	/** Interface for dealing with a web browser window. */
@@ -448,6 +448,6 @@ private:
 	FOnConsoleMessageDelegate OnConsoleMessage;
 
 protected:
-	bool HandleSuppressContextMenu();
+	WEBBROWSER_API bool HandleSuppressContextMenu();
 
-};
\ No newline at end of file
+};
diff --git a/Engine/Source/Runtime/WebBrowser/Public/WebBrowserModule.h b/Engine/Source/Runtime/WebBrowser/Public/WebBrowserModule.h
--- a/Engine/Source/Runtime/WebBrowser/Public/WebBrowserModule.h
+++ b/Engine/Source/Runtime/WebBrowser/Public/WebBrowserModule.h
@@ -11,13 +11,13 @@ class IWebBrowserSingleton;
 /**
  * WebBrowser initialization settings, can be used to override default init behaviors.
  */
-struct WEBBROWSER_API FWebBrowserInitSettings
+struct FWebBrowserInitSettings
 {
 public:
 	/**
 	 * Default constructor. Initializes all members with default behavior values.
 	 */
-	FWebBrowserInitSettings();
+	WEBBROWSER_API FWebBrowserInitSettings();
 
 	// The string which is appended to the browser's user-agent value.
 	FString ProductVersion;
diff --git a/Engine/Source/Runtime/WebBrowser/Public/WebBrowserViewport.h b/Engine/Source/Runtime/WebBrowser/Public/WebBrowserViewport.h
--- a/Engine/Source/Runtime/WebBrowser/Public/WebBrowserViewport.h
+++ b/Engine/Source/Runtime/WebBrowser/Public/WebBrowserViewport.h
@@ -13,7 +13,7 @@ class IWebBrowserWindow;
 /**
  * A Slate viewport to display a Web Browser Window
  */
-class WEBBROWSER_API FWebBrowserViewport : public ISlateViewport
+class FWebBrowserViewport : public ISlateViewport
 {
 public:
 	/**
@@ -36,24 +36,24 @@ public:
 	}
 
 	// ISlateViewport interface
-	virtual FIntPoint GetSize() const override;
-	virtual FSlateShaderResource* GetViewportRenderTargetTexture() const override;
-	virtual void Tick( const FGeometry& AllottedGeometry, double InCurrentTime, float DeltaTime ) override;
-	virtual bool RequiresVsync() const override;
-	virtual FCursorReply OnCursorQuery( const FGeometry& MyGeometry, const FPointerEvent& CursorEvent ) override;
-	virtual FReply OnMouseButtonDown( const FGeometry& MyGeometry, const FPointerEvent& MouseEvent ) override;
-	virtual FReply OnMouseButtonUp( const FGeometry& MyGeometry, const FPointerEvent& MouseEvent ) override;
-	virtual void OnMouseEnter( const FGeometry& MyGeometry, const FPointerEvent& MouseEvent ) override;
-	virtual void OnMouseLeave( const FPointerEvent& MouseEvent ) override;
-	virtual FReply OnMouseMove( const FGeometry& MyGeometry, const FPointerEvent& MouseEvent ) override;
-	virtual FReply OnMouseWheel( const FGeometry& MyGeometry, const FPointerEvent& MouseEvent ) override;
-	virtual FReply OnTouchGesture( const FGeometry& MyGeometry, const FPointerEvent& GestureEvent ) override;
-	virtual FReply OnMouseButtonDoubleClick( const FGeometry& InMyGeometry, const FPointerEvent& InMouseEvent ) override;
-	virtual FReply OnKeyDown( const FGeometry& MyGeometry, const FKeyEvent& InKeyEvent ) override;
-	virtual FReply OnKeyUp( const FGeometry& MyGeometry, const FKeyEvent& InKeyEvent ) override;
-	virtual FReply OnKeyChar( const FGeometry& MyGeometry, const FCharacterEvent& InCharacterEvent ) override;
-	virtual FReply OnFocusReceived( const FFocusEvent& InFocusEvent ) override;
-	virtual void OnFocusLost( const FFocusEvent& InFocusEvent ) override;
+	WEBBROWSER_API virtual FIntPoint GetSize() const override;
+	WEBBROWSER_API virtual FSlateShaderResource* GetViewportRenderTargetTexture() const override;
+	WEBBROWSER_API virtual void Tick( const FGeometry& AllottedGeometry, double InCurrentTime, float DeltaTime ) override;
+	WEBBROWSER_API virtual bool RequiresVsync() const override;
+	WEBBROWSER_API virtual FCursorReply OnCursorQuery( const FGeometry& MyGeometry, const FPointerEvent& CursorEvent ) override;
+	WEBBROWSER_API virtual FReply OnMouseButtonDown( const FGeometry& MyGeometry, const FPointerEvent& MouseEvent ) override;
+	WEBBROWSER_API virtual FReply OnMouseButtonUp( const FGeometry& MyGeometry, const FPointerEvent& MouseEvent ) override;
+	WEBBROWSER_API virtual void OnMouseEnter( const FGeometry& MyGeometry, const FPointerEvent& MouseEvent ) override;
+	WEBBROWSER_API virtual void OnMouseLeave( const FPointerEvent& MouseEvent ) override;
+	WEBBROWSER_API virtual FReply OnMouseMove( const FGeometry& MyGeometry, const FPointerEvent& MouseEvent ) override;
+	WEBBROWSER_API virtual FReply OnMouseWheel( const FGeometry& MyGeometry, const FPointerEvent& MouseEvent ) override;
+	WEBBROWSER_API virtual FReply OnTouchGesture( const FGeometry& MyGeometry, const FPointerEvent& GestureEvent ) override;
+	WEBBROWSER_API virtual FReply OnMouseButtonDoubleClick( const FGeometry& InMyGeometry, const FPointerEvent& InMouseEvent ) override;
+	WEBBROWSER_API virtual FReply OnKeyDown( const FGeometry& MyGeometry, const FKeyEvent& InKeyEvent ) override;
+	WEBBROWSER_API virtual FReply OnKeyUp( const FGeometry& MyGeometry, const FKeyEvent& InKeyEvent ) override;
+	WEBBROWSER_API virtual FReply OnKeyChar( const FGeometry& MyGeometry, const FCharacterEvent& InCharacterEvent ) override;
+	WEBBROWSER_API virtual FReply OnFocusReceived( const FFocusEvent& InFocusEvent ) override;
+	WEBBROWSER_API virtual void OnFocusLost( const FFocusEvent& InFocusEvent ) override;
 
 private:
 	/** The web browser this viewport will display */
diff --git a/Engine/Source/Runtime/WebBrowser/Public/WebJSFunction.h b/Engine/Source/Runtime/WebBrowser/Public/WebJSFunction.h
--- a/Engine/Source/Runtime/WebBrowser/Public/WebJSFunction.h
+++ b/Engine/Source/Runtime/WebBrowser/Public/WebJSFunction.h
@@ -22,7 +22,7 @@ class FWebJSScripting;
 class UObject;
 class UStruct;
 
-struct WEBBROWSER_API FWebJSParam
+struct FWebJSParam
 {
 
 	struct IStructWrapper
@@ -108,9 +108,9 @@ struct WEBBROWSER_API FWebJSParam
 			MapValue->Add(Pair.Key.ToString(), FWebJSParam(Pair.Value));
 		}
 	}
-	FWebJSParam(const FWebJSParam& Other);
-	FWebJSParam(FWebJSParam&& Other);
-	~FWebJSParam();
+	WEBBROWSER_API FWebJSParam(const FWebJSParam& Other);
+	WEBBROWSER_API FWebJSParam(FWebJSParam&& Other);
+	WEBBROWSER_API ~FWebJSParam();
 
 	enum { PTYPE_NULL, PTYPE_BOOL, PTYPE_INT, PTYPE_DOUBLE, PTYPE_STRING, PTYPE_OBJECT, PTYPE_STRUCT, PTYPE_ARRAY, PTYPE_MAP } Tag;
 	union
@@ -131,7 +131,7 @@ class FWebJSScripting;
 
 /** Base class for JS callback objects. */
 USTRUCT()
-struct WEBBROWSER_API FWebJSCallbackBase
+struct FWebJSCallbackBase
 {
 	GENERATED_USTRUCT_BODY()
 	FWebJSCallbackBase()
@@ -149,7 +149,7 @@ protected:
 		, CallbackId(InCallbackId)
 	{}
 
-	void Invoke(int32 ArgCount, FWebJSParam Arguments[], bool bIsError = false) const;
+	WEBBROWSER_API void Invoke(int32 ArgCount, FWebJSParam Arguments[], bool bIsError = false) const;
 
 private:
 
@@ -164,7 +164,7 @@ private:
  * FWebJSFunction objects can also be added to delegates and events using the Bind/AddLambda method.
  */
 USTRUCT()
-struct WEBBROWSER_API FWebJSFunction
+struct FWebJSFunction
 	: public FWebJSCallbackBase
 {
 	GENERATED_USTRUCT_BODY()
@@ -194,7 +194,7 @@ struct WEBBROWSER_API FWebJSFunction
  *  Note that the remote object will become invalid as soon as a result has been delivered, so you can only call either Success or Failure once.
  */
 USTRUCT()
-struct WEBBROWSER_API FWebJSResponse
+struct FWebJSResponse
 	: public FWebJSCallbackBase
 {
 	GENERATED_USTRUCT_BODY()
diff --git a/Engine/Source/Runtime/WidgetCarousel/Public/SCarouselNavigationBar.h b/Engine/Source/Runtime/WidgetCarousel/Public/SCarouselNavigationBar.h
--- a/Engine/Source/Runtime/WidgetCarousel/Public/SCarouselNavigationBar.h
+++ b/Engine/Source/Runtime/WidgetCarousel/Public/SCarouselNavigationBar.h
@@ -16,7 +16,7 @@ class SHorizontalBox;
 /**
 * A horizontal bar of buttons for navigating to a specific item in the widget carousel.
 */
-class WIDGETCAROUSEL_API SCarouselNavigationBar : public SCompoundWidget
+class SCarouselNavigationBar : public SCompoundWidget
 {
 	DECLARE_DELEGATE_OneParam(FOnSelectedIndexChanged, int32)
 
@@ -46,11 +46,11 @@ public:
 
 	SLATE_END_ARGS()
 
-	int32 GetItemCount();
+	WIDGETCAROUSEL_API int32 GetItemCount();
 
-	void SetItemCount(int32 InItemCount);
+	WIDGETCAROUSEL_API void SetItemCount(int32 InItemCount);
 
-	void Construct(const FArguments& InArgs);
+	WIDGETCAROUSEL_API void Construct(const FArguments& InArgs);
 
 private:
 	void BuildScrollBar();
diff --git a/Engine/Source/Runtime/WidgetCarousel/Public/WidgetCarouselModule.h b/Engine/Source/Runtime/WidgetCarousel/Public/WidgetCarouselModule.h
--- a/Engine/Source/Runtime/WidgetCarousel/Public/WidgetCarouselModule.h
+++ b/Engine/Source/Runtime/WidgetCarousel/Public/WidgetCarouselModule.h
@@ -7,12 +7,12 @@
 /**
  * Implements the WidgetCarousel module.
  */
-class WIDGETCAROUSEL_API FWidgetCarouselModule
+class FWidgetCarouselModule
 	: public IModuleInterface
 {
 public:
 
-	virtual void StartupModule( ) override;
+	WIDGETCAROUSEL_API virtual void StartupModule( ) override;
 
-	virtual void ShutdownModule( ) override;
+	WIDGETCAROUSEL_API virtual void ShutdownModule( ) override;
 };
diff --git a/Engine/Source/Runtime/WidgetCarousel/Public/WidgetCarouselStyle.h b/Engine/Source/Runtime/WidgetCarousel/Public/WidgetCarouselStyle.h
--- a/Engine/Source/Runtime/WidgetCarousel/Public/WidgetCarouselStyle.h
+++ b/Engine/Source/Runtime/WidgetCarousel/Public/WidgetCarouselStyle.h
@@ -11,20 +11,20 @@
 #include "WidgetCarouselStyle.generated.h"
 
 USTRUCT()
-struct WIDGETCAROUSEL_API FWidgetCarouselNavigationButtonStyle : public FSlateWidgetStyle
+struct FWidgetCarouselNavigationButtonStyle : public FSlateWidgetStyle
 {
 	GENERATED_USTRUCT_BODY()
 
-	FWidgetCarouselNavigationButtonStyle();
+	WIDGETCAROUSEL_API FWidgetCarouselNavigationButtonStyle();
 
 	virtual ~FWidgetCarouselNavigationButtonStyle() {}
 
-	virtual void GetResources(TArray<const FSlateBrush*> & OutBrushes) const override;
+	WIDGETCAROUSEL_API virtual void GetResources(TArray<const FSlateBrush*> & OutBrushes) const override;
 
-	static const FName TypeName;
+	static WIDGETCAROUSEL_API const FName TypeName;
 	virtual const FName GetTypeName() const override { return TypeName; };
 
-	static const FWidgetCarouselNavigationButtonStyle& GetDefault();
+	static WIDGETCAROUSEL_API const FWidgetCarouselNavigationButtonStyle& GetDefault();
 
 	UPROPERTY()
 	FButtonStyle InnerButtonStyle;
@@ -40,20 +40,20 @@ struct WIDGETCAROUSEL_API FWidgetCarouselNavigationButtonStyle : public FSlateWi
 };
 
 USTRUCT(BlueprintType)
-struct WIDGETCAROUSEL_API FWidgetCarouselNavigationBarStyle : public FSlateWidgetStyle
+struct FWidgetCarouselNavigationBarStyle : public FSlateWidgetStyle
 {
 	GENERATED_USTRUCT_BODY()
 
-	FWidgetCarouselNavigationBarStyle();
+	WIDGETCAROUSEL_API FWidgetCarouselNavigationBarStyle();
 
 	virtual ~FWidgetCarouselNavigationBarStyle() {}
 
-	virtual void GetResources(TArray<const FSlateBrush*> & OutBrushes) const override;
+	WIDGETCAROUSEL_API virtual void GetResources(TArray<const FSlateBrush*> & OutBrushes) const override;
 
-	static const FName TypeName;
+	static WIDGETCAROUSEL_API const FName TypeName;
 	virtual const FName GetTypeName() const override { return TypeName; };
 
-	static const FWidgetCarouselNavigationBarStyle& GetDefault();
+	static WIDGETCAROUSEL_API const FWidgetCarouselNavigationBarStyle& GetDefault();
 
 	UPROPERTY(EditAnywhere, BlueprintReadWrite, Category=Appearance)
 	FSlateBrush HighlightBrush;
@@ -73,21 +73,21 @@ struct WIDGETCAROUSEL_API FWidgetCarouselNavigationBarStyle : public FSlateWidge
 };
 
 /** Manages the style which provides resources for the add content dialog. */
-class WIDGETCAROUSEL_API FWidgetCarouselModuleStyle
+class FWidgetCarouselModuleStyle
 {
 public:
 
-	static void Initialize();
+	static WIDGETCAROUSEL_API void Initialize();
 
-	static void Shutdown();
+	static WIDGETCAROUSEL_API void Shutdown();
 
 	/** reloads textures used by slate renderer */
-	static void ReloadTextures();
+	static WIDGETCAROUSEL_API void ReloadTextures();
 
 	/** @return The Slate style set for the widget carousel */
-	static const ISlateStyle& Get();
+	static WIDGETCAROUSEL_API const ISlateStyle& Get();
 
-	static FName GetStyleSetName();
+	static WIDGETCAROUSEL_API FName GetStyleSetName();
 
 private:
 
diff --git a/Engine/Source/Runtime/Windows/D3D11RHI/Public/D3D11Resources.h b/Engine/Source/Runtime/Windows/D3D11RHI/Public/D3D11Resources.h
--- a/Engine/Source/Runtime/Windows/D3D11RHI/Public/D3D11Resources.h
+++ b/Engine/Source/Runtime/Windows/D3D11RHI/Public/D3D11Resources.h
@@ -150,10 +150,10 @@ private:
 };
 
 /** Texture base class. */
-class D3D11RHI_API FD3D11Texture final : public FRHITexture, public FD3D11ViewableResource
+class FD3D11Texture final : public FRHITexture, public FD3D11ViewableResource
 {
 public:
-	explicit FD3D11Texture(
+	D3D11RHI_API explicit FD3D11Texture(
 		const FRHITextureCreateDesc& InDesc,
 		ID3D11Resource* InResource,
 		ID3D11ShaderResourceView* InShaderResourceView,
@@ -164,10 +164,10 @@ public:
 	);
 
 	enum EAliasResourceParam { CreateAlias };
-	explicit FD3D11Texture(FD3D11Texture const& Other, const FString& Name, EAliasResourceParam);
-	void AliasResource(FD3D11Texture const& Other);
+	D3D11RHI_API explicit FD3D11Texture(FD3D11Texture const& Other, const FString& Name, EAliasResourceParam);
+	D3D11RHI_API void AliasResource(FD3D11Texture const& Other);
 
-	virtual ~FD3D11Texture();
+	D3D11RHI_API virtual ~FD3D11Texture();
 
 	inline uint64 GetMemorySize() const
 	{
@@ -279,10 +279,10 @@ public:
 	* Locks one of the texture's mip-maps.
 	* @return A pointer to the specified texture data.
 	*/
-	void* Lock(class FD3D11DynamicRHI* D3DRHI, uint32 MipIndex, uint32 ArrayIndex, EResourceLockMode LockMode, uint32& DestStride, bool bForceLockDeferred = false, uint64* OutLockedByteCount = nullptr);
+	D3D11RHI_API void* Lock(class FD3D11DynamicRHI* D3DRHI, uint32 MipIndex, uint32 ArrayIndex, EResourceLockMode LockMode, uint32& DestStride, bool bForceLockDeferred = false, uint64* OutLockedByteCount = nullptr);
 
 	/** Unlocks a previously locked mip-map. */
-	void Unlock(class FD3D11DynamicRHI* D3DRHI, uint32 MipIndex, uint32 ArrayIndex);
+	D3D11RHI_API void Unlock(class FD3D11DynamicRHI* D3DRHI, uint32 MipIndex, uint32 ArrayIndex);
 
 private:
 	//Resource handle for use by IHVs for SLI and other purposes.
diff --git a/Engine/Source/Runtime/Windows/D3D11RHI/Public/D3D11Viewport.h b/Engine/Source/Runtime/Windows/D3D11RHI/Public/D3D11Viewport.h
--- a/Engine/Source/Runtime/Windows/D3D11RHI/Public/D3D11Viewport.h
+++ b/Engine/Source/Runtime/Windows/D3D11RHI/Public/D3D11Viewport.h
@@ -11,25 +11,25 @@
 #include "RenderUtils.h"
 
 /** A D3D event query resource. */
-class D3D11RHI_API FD3D11EventQuery
+class FD3D11EventQuery
 {
 public:
 
 	/** Initialization constructor. */
-	FD3D11EventQuery(class FD3D11DynamicRHI* InD3DRHI);
+	D3D11RHI_API FD3D11EventQuery(class FD3D11DynamicRHI* InD3DRHI);
 
 	/** Issues an event for the query to poll. */
-	void IssueEvent();
+	D3D11RHI_API void IssueEvent();
 
 	/** Waits for the event query to finish. */
-	void WaitForCompletion();
+	D3D11RHI_API void WaitForCompletion();
 
 private:
 	FD3D11DynamicRHI* D3DRHI;
 	TRefCountPtr<ID3D11Query> Query;
 };
 
-class D3D11RHI_API FD3D11Viewport : public FRHIViewport
+class FD3D11Viewport : public FRHIViewport
 {
 public:
 	enum ED3DViewportValidFlags : uint32
@@ -40,28 +40,28 @@ public:
 	};
 
 	FD3D11Viewport(class FD3D11DynamicRHI* InD3DRHI) : D3DRHI(InD3DRHI), PresentFailCount(0), ValidState (0), FrameSyncEvent(InD3DRHI) {}
-	FD3D11Viewport(class FD3D11DynamicRHI* InD3DRHI, HWND InWindowHandle, uint32 InSizeX, uint32 InSizeY, bool bInIsFullscreen, EPixelFormat InPreferredPixelFormat);
-	~FD3D11Viewport();
+	D3D11RHI_API FD3D11Viewport(class FD3D11DynamicRHI* InD3DRHI, HWND InWindowHandle, uint32 InSizeX, uint32 InSizeY, bool bInIsFullscreen, EPixelFormat InPreferredPixelFormat);
+	D3D11RHI_API ~FD3D11Viewport();
 
-	virtual void Resize(uint32 InSizeX, uint32 InSizeY, bool bInIsFullscreen, EPixelFormat PreferredPixelFormat);
+	D3D11RHI_API virtual void Resize(uint32 InSizeX, uint32 InSizeY, bool bInIsFullscreen, EPixelFormat PreferredPixelFormat);
 
 	/**
 	 * If the swap chain has been invalidated by DXGI, resets the swap chain to the expected state; otherwise, does nothing.
 	 * Called once/frame by the game thread on all viewports.
 	 * @param bIgnoreFocus - Whether the reset should happen regardless of whether the window is focused.
 	 */
-	void ConditionalResetSwapChain(bool bIgnoreFocus);
+	D3D11RHI_API void ConditionalResetSwapChain(bool bIgnoreFocus);
 
 	/**
 	 * Called whenever the Viewport is moved to see if it has moved between HDR or LDR monitors
 	 */
-	void CheckHDRMonitorStatus();
+	D3D11RHI_API void CheckHDRMonitorStatus();
 
 
 	/** Presents the swap chain. 
 	 * Returns true if Present was done by Engine.
 	 */
-	bool Present(bool bLockToVsync);
+	D3D11RHI_API bool Present(bool bLockToVsync);
 
 	// Accessors.
 	FIntPoint GetSizeXY() const { return FIntPoint(SizeX, SizeY); }
@@ -91,7 +91,7 @@ public:
 	virtual FRHICustomPresent* GetCustomPresent() const { return CustomPresent; }
 
 	virtual void* GetNativeWindow(void** AddParam = nullptr) const override { return (void*)WindowHandle; }
-	static FD3D11Texture* GetSwapChainSurface(FD3D11DynamicRHI* D3DRHI, EPixelFormat PixelFormat, uint32 SizeX, uint32 SizeY, IDXGISwapChain* SwapChain);
+	static D3D11RHI_API FD3D11Texture* GetSwapChainSurface(FD3D11DynamicRHI* D3DRHI, EPixelFormat PixelFormat, uint32 SizeX, uint32 SizeY, IDXGISwapChain* SwapChain);
 
 	static DXGI_FORMAT GetRenderTargetFormat(EPixelFormat PixelFormat)
 	{
@@ -110,19 +110,19 @@ public:
 
 protected:
 
-	void ResetSwapChainInternal(bool bIgnoreFocus);
+	D3D11RHI_API void ResetSwapChainInternal(bool bIgnoreFocus);
 
 	/** Gets the swap chain flags */
-	uint32 GetSwapChainFlags();
+	D3D11RHI_API uint32 GetSwapChainFlags();
 
 	/** Presents the frame synchronizing with DWM. */
-	void PresentWithVsyncDWM();
+	D3D11RHI_API void PresentWithVsyncDWM();
 
 	/**
 	 * Presents the swap chain checking the return result. 
 	 * Returns true if Present was done by Engine.
 	 */
-	bool PresentChecked(int32 SyncInterval);
+	D3D11RHI_API bool PresentChecked(int32 SyncInterval);
 
 	FD3D11DynamicRHI* D3DRHI;
 	uint64 LastFlipTime;
@@ -144,7 +144,7 @@ protected:
 	bool bIsFullscreen;
 	bool bAllowTearing;
 
-	static uint32 GSwapChainFlags;
+	static D3D11RHI_API uint32 GSwapChainFlags;
 
 	TRefCountPtr<IDXGISwapChain> SwapChain;
 	TRefCountPtr<FD3D11Texture> BackBuffer;
@@ -163,7 +163,7 @@ protected:
 
 	FCustomPresentRHIRef CustomPresent;
 
-	DXGI_MODE_DESC SetupDXGI_MODE_DESC() const;
+	D3D11RHI_API DXGI_MODE_DESC SetupDXGI_MODE_DESC() const;
 };
 
 template<>
diff --git a/Engine/Source/Runtime/Windows/WindowsPlatformFeatures/Public/WindowsPlatformFeatures.h b/Engine/Source/Runtime/Windows/WindowsPlatformFeatures/Public/WindowsPlatformFeatures.h
--- a/Engine/Source/Runtime/Windows/WindowsPlatformFeatures/Public/WindowsPlatformFeatures.h
+++ b/Engine/Source/Runtime/Windows/WindowsPlatformFeatures/Public/WindowsPlatformFeatures.h
@@ -5,18 +5,18 @@
 #include "CoreMinimal.h"
 #include "PlatformFeatures.h"
 
-class WINDOWSPLATFORMFEATURES_API FWindowsPlatformFeaturesModule : public IPlatformFeaturesModule
+class FWindowsPlatformFeaturesModule : public IPlatformFeaturesModule
 {
 public:
 
 	/** Creates a new instance of the audio device implemented by the module. */
-	FWindowsPlatformFeaturesModule();
+	WINDOWSPLATFORMFEATURES_API FWindowsPlatformFeaturesModule();
 
-	virtual IVideoRecordingSystem* GetVideoRecordingSystem() override;
-	virtual class ISaveGameSystem* GetSaveGameSystem() override;
+	WINDOWSPLATFORMFEATURES_API virtual IVideoRecordingSystem* GetVideoRecordingSystem() override;
+	WINDOWSPLATFORMFEATURES_API virtual class ISaveGameSystem* GetSaveGameSystem() override;
 
 private:
-	void StartupModule() override;
+	WINDOWSPLATFORMFEATURES_API void StartupModule() override;
 
 };
 
diff --git a/Engine/Source/Runtime/XmlParser/Public/XmlFile.h b/Engine/Source/Runtime/XmlParser/Public/XmlFile.h
--- a/Engine/Source/Runtime/XmlParser/Public/XmlFile.h
+++ b/Engine/Source/Runtime/XmlParser/Public/XmlFile.h
@@ -22,7 +22,7 @@ namespace EConstructMethod
 }
 
 /** Class representing an XML file */
-class XMLPARSER_API FXmlFile
+class FXmlFile
 {
 public:
 
@@ -34,7 +34,7 @@ public:
 	 * @param	InFile				The path/text to use
 	 * @param	ConstructMethod		Whether to load a file of use the string as a buffer of xml data
 	 */
-	FXmlFile(const FString& InFile, EConstructMethod::Type ConstructMethod = EConstructMethod::ConstructFromFile);
+	XMLPARSER_API FXmlFile(const FString& InFile, EConstructMethod::Type ConstructMethod = EConstructMethod::ConstructFromFile);
 	~FXmlFile() { Clear(); };
 
 	FXmlFile(const FXmlFile& rhs) = delete;
@@ -46,58 +46,58 @@ public:
 	 * @param	Path				The path/text to use
 	 * @param	ConstructMethod		Whether to load a file of use the string as a buffer of xml data
 	 */
-	bool LoadFile(const FString& Path, EConstructMethod::Type ConstructMethod = EConstructMethod::ConstructFromFile);
+	XMLPARSER_API bool LoadFile(const FString& Path, EConstructMethod::Type ConstructMethod = EConstructMethod::ConstructFromFile);
 	/** Gets the last error message from the class */
-	FString GetLastError() const;
+	XMLPARSER_API FString GetLastError() const;
 	/** Clears the file of all internals. Note: Makes any existing pointers to FXmlNode's INVALID */
-	void Clear();
+	XMLPARSER_API void Clear();
 	/** Checks to see if a file is loaded */
-	bool IsValid() const;
+	XMLPARSER_API bool IsValid() const;
 
 	/** 
 	 * Returns the root node of the loaded file. nullptr if no file loaded. 
 	 * It is assumed that there will always be one and only one root node.
 	 * @return						Pointer to root node
 	 */
-	const FXmlNode* GetRootNode() const;
+	XMLPARSER_API const FXmlNode* GetRootNode() const;
 
 	/** 
 	 * Returns the root node of the loaded file. nullptr if no file loaded. 
 	 * It is assumed that there will always be one and only one root node.
 	 * @return						Pointer to root node
 	 */
-	FXmlNode* GetRootNode();
+	XMLPARSER_API FXmlNode* GetRootNode();
 
 	/**
 	 * Write to disk, UTF-16 format only at the moment
 	 * @param	Path				File path to save to
 	 * @return						Whether writing the XML to a file succeeded
 	 */
-	bool Save(const FString& Path);
+	XMLPARSER_API bool Save(const FString& Path);
 
 private:
 
 	/** Pre-processes the file doing anything necessary (such as removing comments) */
-	void PreProcessInput(TArray<FString>& Input);
+	XMLPARSER_API void PreProcessInput(TArray<FString>& Input);
 	/** Whites of the text at the specified locations in a passed-in array */
-	void WhiteOut(TArray<FString>& Input, int32 LineStart, int32 LineEnd, int32 IndexStart, int32 IndexEnd);
+	XMLPARSER_API void WhiteOut(TArray<FString>& Input, int32 LineStart, int32 LineEnd, int32 IndexStart, int32 IndexEnd);
 	/** Converts the passed input into a list of tokens for parsing */
-	void Tokenize(FStringView Input, TArray<FString>& Tokens);
+	XMLPARSER_API void Tokenize(FStringView Input, TArray<FString>& Tokens);
 	/** Converts the passed input into a list of tokens for parsing */
-	TArray<FString> Tokenize(const TArray<FString>& Input);
+	XMLPARSER_API TArray<FString> Tokenize(const TArray<FString>& Input);
 	/** 
 	 * Add an attribute to the passed-in array.
 	 * This makes the assumption that an attribute comes in as one 'token' (i.e. no spaces between tag="value").
 	 */
-	static void AddAttribute(const FString& InToken, TArray<FXmlAttribute>& OutAttributes);
+	static XMLPARSER_API void AddAttribute(const FString& InToken, TArray<FXmlAttribute>& OutAttributes);
 	/** Recursively creates nodes from the list of tokens */
-	FXmlNode* CreateRootNode(TArrayView<const FString> Tokens);
+	XMLPARSER_API FXmlNode* CreateRootNode(TArrayView<const FString> Tokens);
 	/** Hook next ptrs up recursively */
-	void HookUpNextPtrs(FXmlNode* Node);
+	XMLPARSER_API void HookUpNextPtrs(FXmlNode* Node);
 	/** Creates the internal file representation as a bunch of FXmlNode's */
-	void CreateNodes(const TArray<FString>& Tokens);
+	XMLPARSER_API void CreateNodes(const TArray<FString>& Tokens);
 	/** Writes a node hierarchy at the given root to a string */
-	static void WriteNodeHierarchy(const FXmlNode& Node, const FString& Indent, FString& Output);
+	static XMLPARSER_API void WriteNodeHierarchy(const FXmlNode& Node, const FString& Indent, FString& Output);
 
 private:
 
diff --git a/Engine/Source/Runtime/XmlParser/Public/XmlNode.h b/Engine/Source/Runtime/XmlParser/Public/XmlNode.h
--- a/Engine/Source/Runtime/XmlParser/Public/XmlNode.h
+++ b/Engine/Source/Runtime/XmlParser/Public/XmlNode.h
@@ -9,7 +9,7 @@
 /** Forward Decl for implementation */
 class FXmlFile;
 
-class XMLPARSER_API FXmlAttribute
+class FXmlAttribute
 {
 public:
 	FXmlAttribute(const FString& InTag, const FString& InValue)
@@ -19,10 +19,10 @@ public:
 	}
 
 	/** Gets the tag of the attribute */
-	const FString& GetTag() const;
+	XMLPARSER_API const FString& GetTag() const;
 
 	/** Gets the value of the attribute */
-	const FString& GetValue() const;
+	XMLPARSER_API const FString& GetValue() const;
 
 private:
 	/** The tag string */
@@ -33,7 +33,7 @@ private:
 };
 
 /** Xml Node representing a line in an xml file */
-class XMLPARSER_API FXmlNode
+class FXmlNode
 {
 	friend class FXmlFile;
 
@@ -49,28 +49,28 @@ private:
 private:
 
 	/** Recursively deletes the nodes for cleanup */
-	void Delete();
+	XMLPARSER_API void Delete();
 
 public:
 
 	/** Gets the next node in a list of nodes */
-	const FXmlNode* GetNextNode() const;
+	XMLPARSER_API const FXmlNode* GetNextNode() const;
 	/** Gets a list of children nodes */
-	const TArray<FXmlNode*>& GetChildrenNodes() const;
+	XMLPARSER_API const TArray<FXmlNode*>& GetChildrenNodes() const;
 	/** Gets the first child of this node which can be iterated into with GetNextNode */
-	const FXmlNode* GetFirstChildNode() const;
+	XMLPARSER_API const FXmlNode* GetFirstChildNode() const;
 	/** Finds the first child node that contains the specified tag */
-	const FXmlNode* FindChildNode(const FString& InTag) const;
+	XMLPARSER_API const FXmlNode* FindChildNode(const FString& InTag) const;
 	/** Finds the first child node that contains the specified tag */
-	FXmlNode* FindChildNode(const FString& InTag);
+	XMLPARSER_API FXmlNode* FindChildNode(const FString& InTag);
 	/** Gets the tag of the node */
-	const FString& GetTag() const;
+	XMLPARSER_API const FString& GetTag() const;
 	/** Gets the value of the node */
-	const FString& GetContent() const;
+	XMLPARSER_API const FString& GetContent() const;
 	/** Sets the new value of the node */
-	void SetContent(const FString& InContent);
+	XMLPARSER_API void SetContent(const FString& InContent);
 	/** Sets the attributes of the node */
-	void SetAttributes(const TArray<FXmlAttribute>& InAttributes);
+	XMLPARSER_API void SetAttributes(const TArray<FXmlAttribute>& InAttributes);
 
 	/**
 	 * Gets all of the attributes in this node
@@ -83,7 +83,7 @@ public:
 	}
 
 	/** Gets an attribute that corresponds with the passed-in tag */
-	FString GetAttribute(const FString& InTag) const;
+	XMLPARSER_API FString GetAttribute(const FString& InTag) const;
 
 	/**
 	 *  Adds a child node to this node  
@@ -91,7 +91,7 @@ public:
 	 * @param	InContent			(optional) The content of the child node
 	 * @param	InAttributes		(optional) An array of attributes of the child node
 	 */
-	void AppendChildNode(const FString& InTag, const FString& InContent = FString(), const TArray<FXmlAttribute>& InAttributes = TArray<FXmlAttribute>());
+	XMLPARSER_API void AppendChildNode(const FString& InTag, const FString& InContent = FString(), const TArray<FXmlAttribute>& InAttributes = TArray<FXmlAttribute>());
 private:
 
 	/** The list of children nodes */
