Disabled file upload functionality in the embedded web browser to address potential vulnerability.

#localization none
#tests QA tested the preflight build and verified that it was fixed

diff --git a/Engine/Source/Runtime/WebBrowser/Private/CEF/CEFBrowserHandler.cpp b/Engine/Source/Runtime/WebBrowser/Private/CEF/CEFBrowserHandler.cpp
--- a/Engine/Source/Runtime/WebBrowser/Private/CEF/CEFBrowserHandler.cpp
+++ b/Engine/Source/Runtime/WebBrowser/Private/CEF/CEFBrowserHandler.cpp
@@ -838,6 +838,24 @@ void FCEFBrowserHandler::OnDraggableRegionsChanged(CefRefPtr<CefBrowser> Browser
 	}
 }
 
+bool FCEFBrowserHandler::OnFileDialog(CefRefPtr<CefBrowser> Browser,
+	FileDialogMode Mode,
+	const CefString& Title,
+	const CefString& DefaultFilePath,
+	const std::vector<CefString>& AcceptFilters,
+	int SelectedAcceptFilter,
+	CefRefPtr<CefFileDialogCallback> Callback)
+{
+	bool Retval = false;
+	TSharedPtr<FCEFWebBrowserWindow> BrowserWindow = BrowserWindowPtr.Pin();
+	if (BrowserWindow.IsValid())
+	{
+		Retval = BrowserWindow->OnFileDialog(Mode, Title, DefaultFilePath, AcceptFilters, SelectedAcceptFilter, Callback);
+	}
+
+	return Retval;
+}	
+
 CefRefPtr<CefCookieAccessFilter> FCEFBrowserHandler::GetCookieAccessFilter(
 	CefRefPtr<CefBrowser> Browser,
 	CefRefPtr<CefFrame> Frame,
diff --git a/Engine/Source/Runtime/WebBrowser/Private/CEF/CEFBrowserHandler.h b/Engine/Source/Runtime/WebBrowser/Private/CEF/CEFBrowserHandler.h
--- a/Engine/Source/Runtime/WebBrowser/Private/CEF/CEFBrowserHandler.h
+++ b/Engine/Source/Runtime/WebBrowser/Private/CEF/CEFBrowserHandler.h
@@ -38,6 +38,7 @@ class FCEFBrowserHandler
 	, public CefResourceRequestHandler
 	, public CefRequestContextHandler
 	, public CefCookieAccessFilter
+	, public CefDialogHandler
 {
 public:
 
@@ -112,6 +113,11 @@ public:
 		return this;
 	}
 
+	virtual CefRefPtr<CefDialogHandler> GetDialogHandler() override
+	{
+		return this;
+	}
+
 	virtual CefRefPtr<CefCookieAccessFilter> GetCookieAccessFilter(
 		CefRefPtr<CefBrowser> Browser,
 		CefRefPtr<CefFrame> Frame,
@@ -317,6 +323,17 @@ public:
 		CefRefPtr<CefFrame> frame, 
 		const std::vector<CefDraggableRegion>& Regions) override;
 
+public:
+	// CefDialogHandler interface
+
+	virtual bool OnFileDialog(CefRefPtr<CefBrowser> Browser,
+		FileDialogMode Mode,
+		const CefString& Title,
+		const CefString& DefaultFilePath,
+		const std::vector<CefString>& AcceptFilters,
+		int SelectedAcceptFilter,
+		CefRefPtr<CefFileDialogCallback> Callback) override;
+
 public:
 
 	IWebBrowserWindow::FOnBeforePopupDelegate& OnBeforePopup()
diff --git a/Engine/Source/Runtime/WebBrowser/Private/CEF/CEFWebBrowserWindow.cpp b/Engine/Source/Runtime/WebBrowser/Private/CEF/CEFWebBrowserWindow.cpp
--- a/Engine/Source/Runtime/WebBrowser/Private/CEF/CEFWebBrowserWindow.cpp
+++ b/Engine/Source/Runtime/WebBrowser/Private/CEF/CEFWebBrowserWindow.cpp
@@ -1413,6 +1413,13 @@ bool FCEFWebBrowserWindow::OnJSDialog(CefJSDialogHandler::JSDialogType DialogTyp
 	return Retval;
 }
 
+bool FCEFWebBrowserWindow::OnFileDialog(CefDialogHandler::FileDialogMode Mode, const CefString& DialogTitle, const CefString& DefaultFilePath, const std::vector<CefString>& AcceptFilters, int SelectedAcceptFilter, CefRefPtr<CefFileDialogCallback> Callback)
+{
+	// This would prevent a file dialog from opening.
+	UE_LOG(LogWebBrowser, Error, TEXT("FileDialogs are prevented."));
+	return true;
+}
+
 bool FCEFWebBrowserWindow::OnBeforeUnloadDialog(const CefString& MessageText, bool IsReload, CefRefPtr<CefJSDialogCallback> Callback)
 {
 	bool Retval = false;
diff --git a/Engine/Source/Runtime/WebBrowser/Private/CEF/CEFWebBrowserWindow.h b/Engine/Source/Runtime/WebBrowser/Private/CEF/CEFWebBrowserWindow.h
--- a/Engine/Source/Runtime/WebBrowser/Private/CEF/CEFWebBrowserWindow.h
+++ b/Engine/Source/Runtime/WebBrowser/Private/CEF/CEFWebBrowserWindow.h
@@ -423,6 +423,11 @@ private:
 	 */
 	bool OnJSDialog(CefJSDialogHandler::JSDialogType DialogType, const CefString& MessageText, const CefString& DefaultPromptText, CefRefPtr<CefJSDialogCallback> Callback, bool& OutSuppressMessage);
 
+	/**
+	 * Handle showing the file select / upload dialogs
+	 */
+	bool OnFileDialog(CefDialogHandler::FileDialogMode Mode, const CefString& DialogTitle, const CefString& DefaultFilePath, const std::vector<CefString>& AcceptFilters, int SelectedAcceptFilter, CefRefPtr<CefFileDialogCallback> Callback);
+
 	/**
 	 * Handle showing unload confirmation dialogs
 	 */
