WebBrowser: Added OnFocusRecieved handler to SWebBrowserView Slate widget to properly route keyboard focus.

diff --git a/Engine/Source/Runtime/WebBrowser/Private/SWebBrowserView.cpp b/Engine/Source/Runtime/WebBrowser/Private/SWebBrowserView.cpp
--- a/Engine/Source/Runtime/WebBrowser/Private/SWebBrowserView.cpp
+++ b/Engine/Source/Runtime/WebBrowser/Private/SWebBrowserView.cpp
@@ -267,6 +267,27 @@ void SWebBrowserView::HandleWindowDeactivated()
 	}
 }
 
+FReply SWebBrowserView::OnFocusReceived(const FGeometry& MyGeometry, const FFocusEvent& InFocusEvent)
+{
+	FReply Reply = FReply::Handled();
+	if (InFocusEvent.GetCause() != EFocusCause::Cleared)
+	{
+		if (BrowserWidget.IsValid())
+		{
+			Reply.SetUserFocus(BrowserWidget.ToSharedRef(), InFocusEvent.GetCause());
+		}
+		else
+		{
+			Reply.SetUserFocus(this->AsShared());
+		}
+		if (BrowserWindow.IsValid())
+		{
+			BrowserWindow->OnFocus(true, false);
+		}
+	}
+	return Reply;
+}
+
 void SWebBrowserView::HandleWindowActivated()
 {
 	if (BrowserViewport.IsValid())
diff --git a/Engine/Source/Runtime/WebBrowser/Public/SWebBrowserView.h b/Engine/Source/Runtime/WebBrowser/Public/SWebBrowserView.h
--- a/Engine/Source/Runtime/WebBrowser/Public/SWebBrowserView.h
+++ b/Engine/Source/Runtime/WebBrowser/Public/SWebBrowserView.h
@@ -167,6 +167,8 @@ public:
 
 	virtual bool SupportsKeyboardFocus() const override {return true;}
 
+	virtual FReply OnFocusReceived(const FGeometry& MyGeometry, const FFocusEvent& InFocusEvent) override;
+
 	/**
 	 * Construct the widget.
 	 *
