iOS: Fix Invalid URL string on iOS

diff --git a/Engine/Source/Runtime/WebBrowser/Private/IOS/IOSPlatformWebBrowser.cpp b/Engine/Source/Runtime/WebBrowser/Private/IOS/IOSPlatformWebBrowser.cpp
--- a/Engine/Source/Runtime/WebBrowser/Private/IOS/IOSPlatformWebBrowser.cpp
+++ b/Engine/Source/Runtime/WebBrowser/Private/IOS/IOSPlatformWebBrowser.cpp
@@ -310,7 +310,7 @@ class SIOSWebBrowserWidget : public SLeafWidget
 		if (WebBrowserWindowPtr.IsValid())
 		{
 			// Capture vars needed for AsyncTask
-			FString UrlString = Url;
+			NSString* UrlString = [NSString stringWithUTF8String : TCHAR_TO_UTF8(*Url)];
 			TWeakPtr<FWebBrowserWindow> AsyncWebBrowserWindowPtr = WebBrowserWindowPtr;
 
 			// Notify on the game thread
@@ -325,7 +325,7 @@ class SIOSWebBrowserWidget : public SLeafWidget
 						RequestDetails.bIsRedirect = false;
 						RequestDetails.bIsMainFrame = true; // shouldOverrideUrlLoading is only called on the main frame
 
-						BrowserWindow->OnBeforeBrowse().Execute(Url, RequestDetails);
+						BrowserWindow->OnBeforeBrowse().Execute(UrlString, RequestDetails);
 						BrowserWindow->SetTitle("");
 					}
 				}
