IOS WebBrowser can now handle links with target="_blank" via the OnBeforePopup() delegate

[RN] IOS WebBrowser can now handle links with target=_blank via the OnBeforePopup() delegate

diff --git a/Engine/Source/Runtime/WebBrowser/Private/IOS/IOSPlatformWebBrowser.cpp b/Engine/Source/Runtime/WebBrowser/Private/IOS/IOSPlatformWebBrowser.cpp
--- a/Engine/Source/Runtime/WebBrowser/Private/IOS/IOSPlatformWebBrowser.cpp
+++ b/Engine/Source/Runtime/WebBrowser/Private/IOS/IOSPlatformWebBrowser.cpp
@@ -8,6 +8,7 @@
 #include "Widgets/SLeafWidget.h"
 #include "MobileJS/MobileJSScripting.h"
 #include "PlatformHttp.h"
+#include "HAL/PlatformProcess.h"
 
 #import <UIKit/UIKit.h>
 #import <MetalKit/MetalKit.h>
@@ -305,6 +306,16 @@ class SIOSWebBrowserWidget : public SLeafWidget
 		}
 	}
 
+	bool HandleOnBeforePopup(const FString& UrlStr, const FString& FrameName)
+	{
+		TSharedPtr<FWebBrowserWindow> BrowserWindow = WebBrowserWindowPtr.Pin();
+		if (BrowserWindow.IsValid() && BrowserWindow->OnBeforePopup().IsBound())
+		{
+			return BrowserWindow->OnBeforePopup().Execute(UrlStr, FrameName);
+		}
+		return false;
+	}
+
 	bool HandleShouldOverrideUrlLoading(const FString& Url)
 	{
 		if (WebBrowserWindowPtr.IsValid())
@@ -797,10 +808,37 @@ supportsMetal : (bool)InSupportsMetal supportsMetalMRT : (bool)InSupportsMetalMR
 }
 
 #if !PLATFORM_TVOS
+
+- (nullable WKWebView *)webView:(WKWebView *)InWebView createWebViewWithConfiguration:(WKWebViewConfiguration *)InConfiguration forNavigationAction:(WKNavigationAction *)InNavigationAction windowFeatures:(WKWindowFeatures *)InWindowFeatures
+{
+	NSURLRequest *request = InNavigationAction.request;
+	FString UrlStr([[request URL] absoluteString]);
+
+	if (InNavigationAction.targetFrame == nil && !UrlStr.IsEmpty() && FPlatformProcess::CanLaunchURL(*UrlStr))
+	{
+		if (WebBrowserWidget->HandleOnBeforePopup(UrlStr, TEXT("_blank")))
+		{
+			// Launched the URL in external browser, don't create a new webview
+			return nil;
+		}
+	}
+	return nil;
+}
+
 - (void)webView:(WKWebView*)InWebView decidePolicyForNavigationAction : (WKNavigationAction*)InNavigationAction decisionHandler : (void(^)(WKNavigationActionPolicy))InDecisionHandler
 {
 	NSURLRequest *request = InNavigationAction.request;
-	FString UrlStr([[request URL]absoluteString]);
+	FString UrlStr([[request URL] absoluteString]);
+
+	if (InNavigationAction.targetFrame == nil && !UrlStr.IsEmpty() && FPlatformProcess::CanLaunchURL(*UrlStr))
+	{
+		if (WebBrowserWidget->HandleOnBeforePopup(UrlStr, TEXT("_blank")))
+		{
+			// Launched the URL in external browser, don't open the link here too
+			InDecisionHandler(WKNavigationActionPolicyCancel);
+			return;
+		}
+	}
 	
 	WebBrowserWidget->HandleShouldOverrideUrlLoading(UrlStr);
 	InDecisionHandler(WKNavigationActionPolicyAllow);
