PR #10575: Fix WriteJsParam FString type

#rnx

diff --git a/Engine/Source/Runtime/WebBrowser/Private/MobileJS/MobileJSScripting.cpp b/Engine/Source/Runtime/WebBrowser/Private/MobileJS/MobileJSScripting.cpp
--- a/Engine/Source/Runtime/WebBrowser/Private/MobileJS/MobileJSScripting.cpp
+++ b/Engine/Source/Runtime/WebBrowser/Private/MobileJS/MobileJSScripting.cpp
@@ -200,7 +200,7 @@ namespace
 				WriteValue(Writer, Key, Param.IntValue);
 				break;
 			case FWebJSParam::PTYPE_STRING:
-				WriteValue(Writer, Key, Param.StringValue);
+				WriteValue(Writer, Key, *Param.StringValue);
 				break;
 			case FWebJSParam::PTYPE_OBJECT:
 			{
