PR #10296: Implement FMobileJSScripting::ConvertStruct function

#rnx

diff --git a/Engine/Source/Runtime/WebBrowser/Private/MobileJS/MobileJSScripting.cpp b/Engine/Source/Runtime/WebBrowser/Private/MobileJS/MobileJSScripting.cpp
--- a/Engine/Source/Runtime/WebBrowser/Private/MobileJS/MobileJSScripting.cpp
+++ b/Engine/Source/Runtime/WebBrowser/Private/MobileJS/MobileJSScripting.cpp
@@ -11,6 +11,7 @@
 #include "StructDeserializer.h"
 #include "UObject/UnrealType.h"
 #include "Async/Async.h"
+#include "JsonObjectConverter.h"
 
 // For UrlDecode/Encode
 #include "Http.h"
@@ -352,6 +353,16 @@ bool FMobileJSScripting::OnJsMessageReceived(const FString& Command, const TArra
 
 FString FMobileJSScripting::ConvertStruct(UStruct* TypeInfo, const void* StructPtr)
 {
+	TSharedRef<FJsonObject> OutJson(new FJsonObject());
+	if (FJsonObjectConverter::UStructToJsonObject(TypeInfo, StructPtr, OutJson))
+	{
+		FString StringToFill;
+		auto Writer = TJsonWriterFactory<TCHAR, TCondensedJsonPrintPolicy<TCHAR>>::Create(&StringToFill);
+		if (FJsonSerializer::Serialize(OutJson, Writer))
+		{
+			return StringToFill;
+		}
+	}
 	return TEXT("undefined");
 }
 
diff --git a/Engine/Source/Runtime/WebBrowser/WebBrowser.Build.cs b/Engine/Source/Runtime/WebBrowser/WebBrowser.Build.cs
--- a/Engine/Source/Runtime/WebBrowser/WebBrowser.Build.cs
+++ b/Engine/Source/Runtime/WebBrowser/WebBrowser.Build.cs
@@ -37,6 +37,7 @@ public class WebBrowser : ModuleRules
 				{
 					"Engine",
 					"Launch",
+					"JsonUtilities",
 					"WebBrowserTexture"
 				}
 			);
