Add missing WebKit function implementation to handle failed internet connections

#rnx

diff --git a/Engine/Source/Runtime/WebBrowser/Private/IOS/IOSPlatformWebBrowser.cpp b/Engine/Source/Runtime/WebBrowser/Private/IOS/IOSPlatformWebBrowser.cpp
--- a/Engine/Source/Runtime/WebBrowser/Private/IOS/IOSPlatformWebBrowser.cpp
+++ b/Engine/Source/Runtime/WebBrowser/Private/IOS/IOSPlatformWebBrowser.cpp
@@ -873,6 +873,12 @@ supportsMetal : (bool)InSupportsMetal supportsMetalMRT : (bool)InSupportsMetalMR
 //	NSLog(@"didFailNavigation: %@, error %@", CurrentUrl, InError);
 	WebBrowserWidget->HandleReceivedError(InError.code, CurrentUrl);
 }
+-(void)webView:(WKWebView *)InWebView didFailProvisionalNavigation : (WKNavigation *)InNavigation withError : (NSError*)InError
+{
+	NSString* CurrentUrl = [InError.userInfo objectForKey : @"NSErrorFailingURLStringKey"];
+	// NSLog(@"didFailProvisionalNavigation: %@, error %@", CurrentUrl, InError);
+	WebBrowserWidget->HandleReceivedError(InError.code, CurrentUrl);
+}
 #endif
 @end
 
