Fix crash during WebBrowserModule startup if the CEF3 DLLs can't be loaded

#rnx

diff --git a/Engine/Source/Runtime/WebBrowser/Private/WebBrowserModule.cpp b/Engine/Source/Runtime/WebBrowser/Private/WebBrowserModule.cpp
--- a/Engine/Source/Runtime/WebBrowser/Private/WebBrowserModule.cpp
+++ b/Engine/Source/Runtime/WebBrowser/Private/WebBrowserModule.cpp
@@ -55,6 +55,10 @@ void FWebBrowserModule::StartupModule()
 		CEF3Utils::BackupCEF3Logfile(FPaths::ProjectLogDir());
 	}
 	bLoadedCEFModule = CEF3Utils::LoadCEF3Modules(true);
+	if (!bLoadedCEFModule)
+	{
+		return;
+	}
 #if PLATFORM_MAC
 	// Dynamically load the CEF framework library into this dylibs memory space.
 	// CEF now loads function pointers at runtime so we need this to be dylib specific.
