WebBrowser: don't process EWH messages when the browser is closing/closed or when the module is shutting down

Trying to invoke a UObject method at this point could result in crashes due to the UObject being gone.

#tests Verified that the new !IsShuttingDown() check prevents forwarding process messages when we're in the module destructor
#rnx

diff --git a/Engine/Source/Runtime/WebBrowser/Private/CEF/CEFBrowserHandler.cpp b/Engine/Source/Runtime/WebBrowser/Private/CEF/CEFBrowserHandler.cpp
--- a/Engine/Source/Runtime/WebBrowser/Private/CEF/CEFBrowserHandler.cpp
+++ b/Engine/Source/Runtime/WebBrowser/Private/CEF/CEFBrowserHandler.cpp
@@ -2,6 +2,7 @@
 
 #include "CEF/CEFBrowserHandler.h"
 #include "HAL/PlatformApplicationMisc.h"
+#include "WebBrowserLog.h"
 
 #if WITH_CEF3
 
@@ -678,6 +679,12 @@ bool FCEFBrowserHandler::OnProcessMessageReceived(CefRefPtr<CefBrowser> Browser,
 	CefProcessId SourceProcess,
 	CefRefPtr<CefProcessMessage> Message)
 {
+	if (IWebBrowserModule::Get().GetSingleton()->IsShuttingDown())
+	{
+		UE_LOG(LogWebBrowser, Warning, TEXT("Received a process message while the module is shutting down, ignoring."));
+		return false;
+	}
+
 	bool Retval = false;
 	FString MessageName = WCHAR_TO_TCHAR(Message->GetName().ToWString().c_str());
 	TSharedPtr<FCEFWebBrowserWindow> BrowserWindow = BrowserWindowPtr.Pin();
diff --git a/Engine/Source/Runtime/WebBrowser/Private/CEF/CEFWebBrowserWindow.cpp b/Engine/Source/Runtime/WebBrowser/Private/CEF/CEFWebBrowserWindow.cpp
--- a/Engine/Source/Runtime/WebBrowser/Private/CEF/CEFWebBrowserWindow.cpp
+++ b/Engine/Source/Runtime/WebBrowser/Private/CEF/CEFWebBrowserWindow.cpp
@@ -2871,6 +2871,24 @@ CefRefPtr<CefDictionaryValue> FCEFWebBrowserWindow::GetProcessInfo()
 
 bool FCEFWebBrowserWindow::OnProcessMessageReceived(CefRefPtr<CefBrowser> Browser, CefRefPtr<CefFrame> frame, CefProcessId SourceProcess, CefRefPtr<CefProcessMessage> Message)
 {
+	if (IsClosing())
+	{
+		UE_LOG(LogWebBrowser, Warning, TEXT("Received a process message for a browser window in closing state, ignoring."));
+		return false;
+	}
+	if (!IsValid())
+	{
+		UE_LOG(LogWebBrowser, Warning, TEXT("Received a process message for a browser window in invalid state, ignoring."));
+		return false;
+	}
+#if CEF_VERSION_MAJOR >= 128
+	if (!Browser || !Browser->IsValid())
+	{
+		UE_LOG(LogWebBrowser, Warning, TEXT("Received a process message with an invalid browser state, ignoring."));
+		return false;
+	}
+#endif
+
 	bool bHandled = Scripting->OnProcessMessageReceived(Browser, SourceProcess, Message);
 
 	if (!bHandled)
diff --git a/Engine/Source/Runtime/WebBrowser/Private/WebBrowserSingleton.cpp b/Engine/Source/Runtime/WebBrowser/Private/WebBrowserSingleton.cpp
--- a/Engine/Source/Runtime/WebBrowser/Private/WebBrowserSingleton.cpp
+++ b/Engine/Source/Runtime/WebBrowser/Private/WebBrowserSingleton.cpp
@@ -480,6 +480,15 @@ FWebBrowserSingleton::~FWebBrowserSingleton()
 #endif
 }
 
+bool FWebBrowserSingleton::IsShuttingDown() const
+{
+#if WITH_CEF3
+	return bAllowCEF && !CEFBrowserApp;
+#else
+	return false;
+#endif
+}
+
 TSharedRef<IWebBrowserWindowFactory> FWebBrowserSingleton::GetWebBrowserWindowFactory() const
 {
 	return WebBrowserWindowFactory;
diff --git a/Engine/Source/Runtime/WebBrowser/Private/WebBrowserSingleton.h b/Engine/Source/Runtime/WebBrowser/Private/WebBrowserSingleton.h
--- a/Engine/Source/Runtime/WebBrowser/Private/WebBrowserSingleton.h
+++ b/Engine/Source/Runtime/WebBrowser/Private/WebBrowserSingleton.h
@@ -65,12 +65,15 @@ public:
 	*/
 	static FString GetCurrentLocaleCode();
 
-	virtual FString ApplicationCacheDir() const override;
 
 public:
 
 	// IWebBrowserSingleton Interface
 
+	virtual bool IsShuttingDown() const override;
+
+	virtual FString ApplicationCacheDir() const override;
+
 	virtual TSharedRef<IWebBrowserWindowFactory> GetWebBrowserWindowFactory() const override;
 
 	TSharedPtr<IWebBrowserWindow> CreateBrowserWindow(
diff --git a/Engine/Source/Runtime/WebBrowser/Public/IWebBrowserSingleton.h b/Engine/Source/Runtime/WebBrowser/Public/IWebBrowserSingleton.h
--- a/Engine/Source/Runtime/WebBrowser/Public/IWebBrowserSingleton.h
+++ b/Engine/Source/Runtime/WebBrowser/Public/IWebBrowserSingleton.h
@@ -93,6 +93,9 @@ public:
 	 */
 	virtual ~IWebBrowserSingleton() {};
 
+	/** Indicates whether the browser engine is in the shutdown phase. */
+	virtual bool IsShuttingDown() const = 0;
+
 	/** @return A factory object that can be used to construct additional WebBrowserWindows on demand. */
 	virtual TSharedRef<IWebBrowserWindowFactory> GetWebBrowserWindowFactory() const = 0;
 
