Add Cvar to allow to inspect webview with Safari developer tools

diff --git a/Engine/Source/Runtime/WebBrowser/Private/IOS/IOSPlatformWebBrowser.cpp b/Engine/Source/Runtime/WebBrowser/Private/IOS/IOSPlatformWebBrowser.cpp
--- a/Engine/Source/Runtime/WebBrowser/Private/IOS/IOSPlatformWebBrowser.cpp
+++ b/Engine/Source/Runtime/WebBrowser/Private/IOS/IOSPlatformWebBrowser.cpp
@@ -16,6 +16,14 @@
 #include "WebBrowserModule.h"
 #include "IWebBrowserSingleton.h"
 
+#if !UE_BUILD_SHIPPING
+TAutoConsoleVariable<bool> CVarWebViewEnableInspectable(
+	TEXT("webview.EnableInspectable"),
+	false,
+	TEXT("Allows webview to be debuggable through Safari Web Inspector."),
+	ECVF_Default
+);
+#endif
 
 class SIOSWebBrowserWidget : public SLeafWidget
 {
@@ -97,6 +105,9 @@ class SIOSWebBrowserWidget : public SLeafWidget
 	{
 		if (WebViewWrapper != nil)
 		{
+#if !UE_BUILD_SHIPPING
+			WebViewWrapper.WebView.inspectable = CVarWebViewEnableInspectable.GetValueOnAnyThread()? YES : NO;
+#endif
 			if (WebBrowserWindowPtr.IsValid())
 			{
 				WebBrowserWindowPtr.Pin()->SetTickLastFrame();
@@ -492,6 +503,9 @@ supportsMetal : (bool)InSupportsMetal supportsMetalMRT : (bool)InSupportsMetalMR
 		[theConfiguration.userContentController addScriptMessageHandler:self name: MessageHandlerName];
 
 		WebView = [[WKWebView alloc]initWithFrame:CGRectMake(1, 1, 100, 100)  configuration : theConfiguration];
+#if !UE_BUILD_SHIPPING
+		WebView.inspectable = CVarWebViewEnableInspectable.GetValueOnAnyThread()? YES : NO;
+#endif
 		[self.WebViewContainer addSubview : WebView];
 		WebView.navigationDelegate = self;
 		WebView.UIDelegate = self;
