Fix build issue in TVOS

diff --git a/Engine/Source/Runtime/WebBrowser/Private/IOS/IOSPlatformWebBrowser.cpp b/Engine/Source/Runtime/WebBrowser/Private/IOS/IOSPlatformWebBrowser.cpp
--- a/Engine/Source/Runtime/WebBrowser/Private/IOS/IOSPlatformWebBrowser.cpp
+++ b/Engine/Source/Runtime/WebBrowser/Private/IOS/IOSPlatformWebBrowser.cpp
@@ -105,9 +105,6 @@ class SIOSWebBrowserWidget : public SLeafWidget
 	{
 		if (WebViewWrapper != nil)
 		{
-#if !UE_BUILD_SHIPPING
-			WebViewWrapper.WebView.inspectable = CVarWebViewEnableInspectable.GetValueOnAnyThread()? YES : NO;
-#endif
 			if (WebBrowserWindowPtr.IsValid())
 			{
 				WebBrowserWindowPtr.Pin()->SetTickLastFrame();
@@ -135,6 +132,12 @@ class SIOSWebBrowserWidget : public SLeafWidget
 			[WebViewWrapper updateframe : NewFrame];
 
 #if !PLATFORM_TVOS
+#if !UE_BUILD_SHIPPING
+			if (WebViewWrapper.WebView != nil)
+			{
+				WebViewWrapper.WebView.inspectable = CVarWebViewEnableInspectable.GetValueOnAnyThread()? YES : NO;
+			}
+#endif
 			if (IsIOS3DBrowser)
 			{
 				if (WebBrowserTexture)
