Added a configurable Close button to WebBrowser widget public interface.

diff --git a/Engine/Build/Android/Java/src/com/epicgames/unreal/WebViewControl.java b/Engine/Build/Android/Java/src/com/epicgames/unreal/WebViewControl.java
--- a/Engine/Build/Android/Java/src/com/epicgames/unreal/WebViewControl.java
+++ b/Engine/Build/Android/Java/src/com/epicgames/unreal/WebViewControl.java
@@ -6,6 +6,8 @@ import android.graphics.Bitmap;
 import android.graphics.Bitmap.Config;
 import android.graphics.Canvas;
 import android.graphics.Color;
+import android.graphics.Paint;
+import android.graphics.RectF;
 import android.opengl.EGL14;
 import android.opengl.EGLConfig;
 import android.opengl.EGLContext;
@@ -45,6 +47,10 @@ class WebViewPositionLayout extends ViewGroup
 	protected void onLayout(boolean changed, int left, int top, int right, int bottom)
 	{
 		webViewControl.webView.layout(webViewControl.curX,webViewControl.curY,webViewControl.curX+webViewControl.curW,webViewControl.curY+webViewControl.curH);
+		if (webViewControl.closeButton != null)
+		{
+			webViewControl.closeButton.layout(webViewControl.curX,webViewControl.curY,webViewControl.curX+webViewControl.curW,webViewControl.curY+webViewControl.curH);
+		}
 	}
 
 	private WebViewControl webViewControl;
@@ -82,7 +88,7 @@ class WebViewControl
 		public float VOffset;
 	}
 
-	public WebViewControl(long inNativePtr, int width, int height, boolean swizzlePixels, boolean vulkanRenderer, final boolean bEnableRemoteDebugging, final boolean bUseTransparency, final boolean bEnableDomStorage, final boolean bShouldUseBitmapRender, final String userAgentApplication)
+	public WebViewControl(long inNativePtr, int width, int height, boolean swizzlePixels, boolean vulkanRenderer, final boolean bEnableRemoteDebugging, final boolean bUseTransparency, final boolean bEnableDomStorage, final boolean bShouldUseBitmapRender, final String userAgentApplication, boolean bEnableFloatingCloseButton)
 	{
 		final WebViewControl w = this;
 
@@ -145,10 +151,16 @@ class WebViewControl
 
 				// Wrap the webview in a layout that will do absolute positioning for us
 				positionLayout = new WebViewPositionLayout(GameActivity.Get(), w);
-				ViewGroup.LayoutParams params = new ViewGroup.LayoutParams(ViewGroup.LayoutParams.FILL_PARENT, ViewGroup.LayoutParams.FILL_PARENT);
+				ViewGroup.LayoutParams params = new ViewGroup.LayoutParams(ViewGroup.LayoutParams.MATCH_PARENT, ViewGroup.LayoutParams.MATCH_PARENT);
 
 				positionLayout.addView(webView, params);
 
+				if (bEnableFloatingCloseButton)
+				{
+					closeButton = new CloseButton(GameActivity.Get(), positionLayout, w);
+					closeButton.setElevation(webView.getElevation() + 0.1f);
+				}
+
 				bShown = false;
 				NextURL = null;
 				NextContent = null;
@@ -288,6 +300,25 @@ class WebViewControl
 		});
 	}
 
+	public void ShowFloatingCloseButton(boolean bShow, boolean bDraggable)
+	{
+		GameActivity.Get().runOnUiThread(new Runnable()
+		{
+			@Override
+			public void run()
+			{
+				if (closeButton == null)
+				{
+					GameActivity.Log.warn("Close button not enabled in config. Add config for Engine:[Browser]bEnableFloatingCloseButton=true");
+				}
+				else
+				{
+					closeButton.ShowFloatingCloseButton(bShow, bDraggable);
+				}
+			}
+		});
+	}
+
 	public void GoBackOrForward(final int Steps)
 	{
 		GameActivity.Get().runOnUiThread(new Runnable()
@@ -516,7 +547,7 @@ class WebViewControl
 					bShown = true;
 					
 					// add to the activitiy, on top of the SurfaceView
-					ViewGroup.LayoutParams params = new ViewGroup.LayoutParams(ViewGroup.LayoutParams.FILL_PARENT, ViewGroup.LayoutParams.FILL_PARENT);			
+					ViewGroup.LayoutParams params = new ViewGroup.LayoutParams(ViewGroup.LayoutParams.MATCH_PARENT, ViewGroup.LayoutParams.MATCH_PARENT);			
 					GameActivity.Get().addContentView(positionLayout, params);
 					if(!webView.IsAndroid3DBrowser && TargetBitmap == null)
 					{
@@ -612,6 +643,10 @@ class WebViewControl
 					{
 						parent.removeView(webView);
 					}
+					if (closeButton != null)
+					{
+						closeButton.destroy();
+					}
 					parent = (ViewGroup)positionLayout.getParent();
 					if (parent != null)
 					{
@@ -624,7 +659,8 @@ class WebViewControl
 		});
 	}
 
-	
+	public native void FloatingCloseButtonPressed();
+
 	// ======================================================================================
 	private boolean CreateBitmapRenderer()
 	{
@@ -1961,6 +1997,7 @@ class WebViewControl
 	}
 		
 	public GLWebView webView;
+	public CloseButton closeButton;
 	private WebViewPositionLayout positionLayout;
 	public int curX, curY, curW, curH;
 	private boolean bShown;
@@ -1970,4 +2007,201 @@ class WebViewControl
 	
 	// Address of the native SAndroidWebBrowserWidget object
 	private long nativePtr;
+
+	public class CloseButton extends View implements View.OnTouchListener {
+
+		private final static float CLICK_DRAG_TOLERANCE = 10;
+		final float Radius = 25f;
+		final float Size = 20f;
+		final float OutlineSize = 2f;
+		final float Margin = 18f; // dp
+		private final RectF ViewRect = new RectF();
+		private final Paint BackgroundPaint;
+		private final Paint CrossPaint;
+		private final Paint CrossOutlinePaint;
+		float RectWidth = 0f;
+		float SafeSpaceWidth = 0f;
+		float SafeSpaceHeight = 0f;
+		private float downRawX, downRawY;
+		private float downDisplacementX, downDisplacementY;
+		private float xDisplacement, yDisplacement;
+		private WebViewControl WebView;
+		private boolean bIsDraggable = false;
+
+		public CloseButton(Context context, ViewGroup parent, WebViewControl InWebView) {
+			super(context);
+
+			WebView = InWebView;
+			BackgroundPaint = new Paint();
+			BackgroundPaint.setColor(Color.argb(20, 20, 20, 20));
+			CrossPaint = new Paint();
+			CrossPaint.setColor(Color.WHITE);
+			CrossPaint.setStrokeWidth(5);
+			CrossOutlinePaint = new Paint();
+			CrossOutlinePaint.setColor(Color.BLACK);
+			CrossOutlinePaint.setStrokeWidth(10);
+			SafeSpaceWidth = pixelsFromDp(Margin);
+			SafeSpaceHeight = pixelsFromDp(Margin);
+
+			parent.addView(this, new ViewGroup.MarginLayoutParams(ViewGroup.LayoutParams.MATCH_PARENT, ViewGroup.LayoutParams.MATCH_PARENT));
+
+			setOutlineProvider(null);
+			ShowFloatingCloseButton(false, false);
+			invalidate();
+		}
+
+		public void destroy() {
+			setOnTouchListener(null);
+			if (getParent() != null) {
+				((ViewGroup) getParent()).removeView(this);
+			}
+		}
+
+		public float pixelsFromDp(final float dp) {
+			return dp * getContext().getResources().getDisplayMetrics().density;
+		}
+
+		@Override
+		public void onSizeChanged(int Width, int Height, int OldWidth, int OldHeight) {
+			xDisplacement = 0;
+			yDisplacement = 0;
+			ViewRect.set(0, 0, Width, Height);
+			RectWidth = ViewRect.width();
+		}
+
+		public float GetOriginX()
+		{
+			return RectWidth - SafeSpaceWidth - 2 * Size + xDisplacement;
+		}
+
+		public float GetOriginY()
+		{
+			return SafeSpaceHeight + yDisplacement;
+		}
+		@Override
+		public void draw(Canvas Canvas) {
+			super.draw(Canvas);
+
+			// Draw background
+			//Canvas.drawCircle(RectWidth - SafeSpaceWidth, SafeSpaceHeight, Radius, BackgroundPaint);
+
+			// Draw cross
+			float originX = GetOriginX();
+			float originY = GetOriginY();
+
+			Canvas.drawLine(
+				originX - OutlineSize,
+				originY - OutlineSize,
+				originX + OutlineSize + Size * 2,
+				originY + OutlineSize + Size * 2,
+				CrossOutlinePaint);
+			Canvas.drawLine(
+				originX - OutlineSize,
+				originY + OutlineSize + Size * 2,
+				originX + OutlineSize + Size * 2,
+				originY - OutlineSize,
+				CrossOutlinePaint);
+			Canvas.drawLine(
+				originX,
+				originY,
+				originX + Size * 2,
+				originY + Size * 2,
+				CrossPaint);
+			Canvas.drawLine(
+				originX,
+				originY + Size * 2,
+				originX + Size * 2,
+				originY,
+				CrossPaint);
+
+			invalidate();
+		}
+
+		public void ShowFloatingCloseButton(boolean bShow, boolean bDraggable)
+		{
+			bIsDraggable = bDraggable;
+			if (bShow)
+			{
+				setOnTouchListener(this);
+				setVisibility(View.VISIBLE);
+			}
+			else
+			{
+				setOnTouchListener(null);
+				setVisibility(View.GONE);
+			}
+		}
+
+		public boolean IsTouchInside(float X, float Y) {
+			float originX = GetOriginX();
+			float originY = GetOriginY();
+
+			RectF rect = new RectF(
+				originX - Size * 0.5f,
+				originY - Size * 0.5f,
+				originX + Size * 2.5f,
+				originY + Size * 2.5f
+			);
+			int[] rawViewOrigin = new int[2];
+			getLocationOnScreen(rawViewOrigin);
+			return rect.contains(X - rawViewOrigin[0], Y - rawViewOrigin[1]);
+		}
+
+		@Override
+		public boolean onTouch(View view, MotionEvent motionEvent) {
+			boolean IsInside = IsTouchInside(motionEvent.getRawX(), motionEvent.getRawY());
+
+			switch (motionEvent.getAction()) {
+				case MotionEvent.ACTION_DOWN: {
+					//add the offset so we can drag again after a displacement
+					downRawX = motionEvent.getRawX() ;
+					downRawY = motionEvent.getRawY();
+
+					downDisplacementX = downRawX - xDisplacement;
+					downDisplacementY = downRawY - yDisplacement;
+
+					if (IsInside)
+					{
+						return true;
+					}
+					return false;
+				}
+				case MotionEvent.ACTION_MOVE: {
+					if (bIsDraggable)
+					{
+						float newXDisplacement = motionEvent.getRawX() - downDisplacementX;
+						float newYDisplacement = motionEvent.getRawY() - downDisplacementY;
+
+						int viewWidth = view.getWidth();
+						int viewHeight = view.getHeight();
+
+						xDisplacement = Math.max(-(viewWidth - 2 * SafeSpaceWidth - 2 * Size), Math.min(0, newXDisplacement));
+						yDisplacement = Math.max(0, Math.min(viewHeight - 2 * SafeSpaceHeight - 2 * Size, newYDisplacement));
+					}
+					return true;
+				}
+				case MotionEvent.ACTION_UP: {
+					float upRawX = motionEvent.getRawX();
+					float upRawY = motionEvent.getRawY();
+
+					float upDX = upRawX - downRawX;
+					float upDY = upRawY - downRawY;
+
+					if (Math.abs(upDX) < CLICK_DRAG_TOLERANCE && Math.abs(upDY) < CLICK_DRAG_TOLERANCE) { 
+						// A click
+						return performClick();
+					}
+
+					return true;
+				}
+			}
+			return false;
+		}
+
+		@Override
+		public boolean performClick() {
+			WebView.FloatingCloseButtonPressed();
+			return super.performClick();
+		}
+	}
 }
diff --git a/Engine/Source/Runtime/WebBrowser/Private/Android/AndroidJavaWebBrowser.cpp b/Engine/Source/Runtime/WebBrowser/Private/Android/AndroidJavaWebBrowser.cpp
--- a/Engine/Source/Runtime/WebBrowser/Private/Android/AndroidJavaWebBrowser.cpp
+++ b/Engine/Source/Runtime/WebBrowser/Private/Android/AndroidJavaWebBrowser.cpp
@@ -33,8 +33,8 @@ static jfieldID FindField(JNIEnv* JEnv, jclass Class, const ANSICHAR* FieldName,
 #define MOTIONEVENT_ACTION_MOVE		2
 
 FJavaAndroidWebBrowser::FJavaAndroidWebBrowser(bool swizzlePixels, bool vulkanRenderer, int32 width, int32 height,
-	jlong widgetPtr, bool bEnableRemoteDebugging, bool bUseTransparency, bool bEnableDomStorage, bool bShouldUseBitmapRender, const FString& UserAgentApplication)
-	: FJavaClassObject(GetClassName(), "(JIIZZZZZZLjava/lang/String;)V", widgetPtr, width, height,  swizzlePixels, vulkanRenderer, bEnableRemoteDebugging, bUseTransparency, bEnableDomStorage, bShouldUseBitmapRender, *FJavaHelper::ToJavaString(AndroidJavaEnv::GetJavaEnv(), UserAgentApplication))
+	jlong widgetPtr, bool bEnableRemoteDebugging, bool bUseTransparency, bool bEnableDomStorage, bool bShouldUseBitmapRender, const FString& UserAgentApplication, bool bEnableFloatingCloseButton)
+	: FJavaClassObject(GetClassName(), "(JIIZZZZZZLjava/lang/String;Z)V", widgetPtr, width, height, swizzlePixels, vulkanRenderer, bEnableRemoteDebugging, bUseTransparency, bEnableDomStorage, bShouldUseBitmapRender, *FJavaHelper::ToJavaString(AndroidJavaEnv::GetJavaEnv(), UserAgentApplication), bEnableFloatingCloseButton)
 	, ReleaseMethod(GetClassMethod("release", "()V"))
 	, GetVideoLastFrameBitmapMethod(GetClassMethod("getVideoLastFrameBitmap", "()Lcom/epicgames/unreal/WebViewControl$FrameUpdateInfo;"))
 	, GetVideoLastFrameDataMethod(GetClassMethod("getVideoLastFrameData", "()Lcom/epicgames/unreal/WebViewControl$FrameUpdateInfo;"))
@@ -48,6 +48,7 @@ FJavaAndroidWebBrowser::FJavaAndroidWebBrowser(bool swizzlePixels, bool vulkanRe
 	, StopLoadMethod(GetClassMethod("StopLoad", "()V"))
 	, ReloadMethod(GetClassMethod("Reload", "()V"))
 	, CloseMethod(GetClassMethod("Close", "()V"))
+	, ShowFloatingCloseButtonMethod(GetClassMethod("ShowFloatingCloseButton", "(ZZ)V"))
 	, GoBackOrForwardMethod(GetClassMethod("GoBackOrForward", "(I)V"))
 	, SendTouchEventMethod(GetClassMethod("SendTouchEvent", "(IFF)V"))
 	, SendKeyEventMethod(GetClassMethod("SendKeyEvent", "(ZI)Z"))
@@ -269,6 +270,11 @@ void FJavaAndroidWebBrowser::Close()
 	CallMethod<void>(CloseMethod);
 }
 
+void FJavaAndroidWebBrowser::ShowFloatingCloseButton(bool bShow, bool bDraggable)
+{
+	CallMethod<void>(ShowFloatingCloseButtonMethod, bShow, bDraggable);
+}
+
 void FJavaAndroidWebBrowser::GoBack()
 {
 	CallMethod<void>(GoBackOrForwardMethod, -1);
diff --git a/Engine/Source/Runtime/WebBrowser/Private/Android/AndroidJavaWebBrowser.h b/Engine/Source/Runtime/WebBrowser/Private/Android/AndroidJavaWebBrowser.h
--- a/Engine/Source/Runtime/WebBrowser/Private/Android/AndroidJavaWebBrowser.h
+++ b/Engine/Source/Runtime/WebBrowser/Private/Android/AndroidJavaWebBrowser.h
@@ -15,7 +15,7 @@
 class FJavaAndroidWebBrowser : public FJavaClassObject
 {
 public:
-	FJavaAndroidWebBrowser(bool swizzlePixels, bool vulkanRenderer, int32 width, int32 height, jlong widgetPtr, bool bEnableRemoteDebugging, bool bUseTransparency, bool bEnableDomStorage, bool bShouldUseBitmapRender, const FString& UserAgentApplication);
+	FJavaAndroidWebBrowser(bool swizzlePixels, bool vulkanRenderer, int32 width, int32 height, jlong widgetPtr, bool bEnableRemoteDebugging, bool bUseTransparency, bool bEnableDomStorage, bool bShouldUseBitmapRender, const FString& UserAgentApplication, bool bEnableFloatingCloseButton);
 	virtual ~FJavaAndroidWebBrowser();
 	void Release();
 	bool GetVideoLastFrameBitmap(void* outPixels, int64 outCount);
@@ -29,6 +29,7 @@ public:
 	void StopLoad();
 	void Reload();
 	void Close();
+	void ShowFloatingCloseButton(bool bShow, bool bDraggable);
 	void GoBack();
 	void GoForward();
 	void SendTouchDown(float x, float y);
@@ -55,6 +56,7 @@ private:
 	FJavaClassMethod StopLoadMethod;
 	FJavaClassMethod ReloadMethod;
 	FJavaClassMethod CloseMethod;
+	FJavaClassMethod ShowFloatingCloseButtonMethod;
 	FJavaClassMethod GoBackOrForwardMethod;
 	FJavaClassMethod SendTouchEventMethod;
 	FJavaClassMethod SendKeyEventMethod;
diff --git a/Engine/Source/Runtime/WebBrowser/Private/Android/AndroidWebBrowserWidget.cpp b/Engine/Source/Runtime/WebBrowser/Private/Android/AndroidWebBrowserWidget.cpp
--- a/Engine/Source/Runtime/WebBrowser/Private/Android/AndroidWebBrowserWidget.cpp
+++ b/Engine/Source/Runtime/WebBrowser/Private/Android/AndroidWebBrowserWidget.cpp
@@ -127,8 +127,11 @@ void SAndroidWebBrowserWidget::Construct(const FArguments& Args)
 	FIntPoint viewportSize = WebBrowserWindowPtr.Pin()->GetViewportSize();
 	FPlatformMisc::LowLevelOutputDebugStringf(TEXT("SAndroidWebBrowserWidget::Construct viewport=%d x %d"), viewportSize.X, viewportSize.Y);
 
+	bool bEnableFloatingCloseButton = false;
+	GConfig->GetBool(TEXT("Browser"), TEXT("bEnableFloatingCloseButton"), bEnableFloatingCloseButton, GEngineIni);
+
 	JavaWebBrowser = MakeShared<FJavaAndroidWebBrowser, ESPMode::ThreadSafe>(false, FAndroidMisc::ShouldUseVulkan(), viewportSize.X, viewportSize.Y,
-		reinterpret_cast<jlong>(this), !(UE_BUILD_SHIPPING || UE_BUILD_TEST), Args._UseTransparency, bEnableDomStorage, bShouldUseBitmapRender, Args._UserAgentApplication);
+		reinterpret_cast<jlong>(this), !(UE_BUILD_SHIPPING || UE_BUILD_TEST), Args._UseTransparency, bEnableDomStorage, bShouldUseBitmapRender, Args._UserAgentApplication, bEnableFloatingCloseButton);
 
 	TextureSamplePool = new FWebBrowserTextureSamplePool();
 	WebBrowserTextureSamplesQueue = MakeShared<FWebBrowserTextureSampleQueue, ESPMode::ThreadSafe>();
@@ -504,6 +507,11 @@ void SAndroidWebBrowserWidget::Close()
 	WebBrowserWindowPtr.Reset();
 }
 
+void SAndroidWebBrowserWidget::ShowFloatingCloseButton(bool bShow, bool bDraggable)
+{
+	JavaWebBrowser->ShowFloatingCloseButton(bShow, bDraggable);
+}
+
 void SAndroidWebBrowserWidget::GoBack()
 {
 	JavaWebBrowser->GoBack();
@@ -843,8 +851,33 @@ void SAndroidWebBrowserWidget::HandleReceivedError(jint ErrorCode, jstring /* ig
 	}
 }
 
+void SAndroidWebBrowserWidget::FloatingCloseButtonPressed()
+{
+	if (WebBrowserWindowPtr.IsValid())
+	{
+		TSharedPtr<FAndroidWebBrowserWindow> BrowserWindow = WebBrowserWindowPtr.Pin();
+		if (BrowserWindow.IsValid())
+		{
+			FGraphEventRef OnRequestClose = FFunctionGraphTask::CreateAndDispatchWhenReady([&]()
+				{
+					BrowserWindow->FloatingCloseButtonPressed();
+				}, TStatId(), NULL, ENamedThreads::GameThread);
+			FTaskGraphInterface::Get().WaitUntilTaskCompletes(OnRequestClose);
+		}
+	}
+}
+
 // Native method implementations:
 
+JNI_METHOD void Java_com_epicgames_unreal_WebViewControl_FloatingCloseButtonPressed(JNIEnv* JEnv, jobject Client)
+{
+	TSharedPtr<SAndroidWebBrowserWidget> Widget = SAndroidWebBrowserWidget::GetWidgetPtr(JEnv, Client);
+	if (Widget.IsValid())
+	{
+		Widget->FloatingCloseButtonPressed();
+	}
+}
+
 JNI_METHOD jbyteArray Java_com_epicgames_unreal_WebViewControl_00024ViewClient_shouldInterceptRequestImpl(JNIEnv* JEnv, jobject Client, jstring JUrl)
 {
 	TSharedPtr<SAndroidWebBrowserWidget> Widget = SAndroidWebBrowserWidget::GetWidgetPtr(JEnv, Client);
diff --git a/Engine/Source/Runtime/WebBrowser/Private/Android/AndroidWebBrowserWidget.h b/Engine/Source/Runtime/WebBrowser/Private/Android/AndroidWebBrowserWidget.h
--- a/Engine/Source/Runtime/WebBrowser/Private/Android/AndroidWebBrowserWidget.h
+++ b/Engine/Source/Runtime/WebBrowser/Private/Android/AndroidWebBrowserWidget.h
@@ -55,10 +55,12 @@ public:
 	void StopLoad();
 	void Reload();
 	void Close();
+	void ShowFloatingCloseButton(bool bShow, bool bDraggable);
 	void GoBack();
 	void GoForward();
 	bool CanGoBack();
 	bool CanGoForward();
+	void FloatingCloseButtonPressed();
 
 	virtual FReply OnMouseButtonDown(const FGeometry& MyGeometry, const FPointerEvent& MouseEvent) override;
 	virtual FReply OnMouseButtonUp(const FGeometry& MyGeometry, const FPointerEvent& MouseEvent) override;
@@ -139,4 +141,4 @@ private:
 	FWebBrowserTextureSamplePool* TextureSamplePool;
 };
 
-#endif // USE_ANDROID_JNI
\ No newline at end of file
+#endif // USE_ANDROID_JNI
diff --git a/Engine/Source/Runtime/WebBrowser/Private/Android/AndroidWebBrowserWindow.cpp b/Engine/Source/Runtime/WebBrowser/Private/Android/AndroidWebBrowserWindow.cpp
--- a/Engine/Source/Runtime/WebBrowser/Private/Android/AndroidWebBrowserWindow.cpp
+++ b/Engine/Source/Runtime/WebBrowser/Private/Android/AndroidWebBrowserWindow.cpp
@@ -345,6 +345,11 @@ void FAndroidWebBrowserWindow::SetParentWindow(TSharedPtr<SWindow> Window)
 	ParentWindow = Window;
 }
 
+void FAndroidWebBrowserWindow::ShowFloatingCloseButton(bool bShow, bool bDraggable)
+{
+	BrowserWidget->ShowFloatingCloseButton(bShow, bDraggable);
+}
+
 void FAndroidWebBrowserWindow::ExecuteJavascript(const FString& Script)
 {
 	BrowserWidget->ExecuteJavascript(Script);
@@ -397,4 +402,12 @@ bool FAndroidWebBrowserWindow::IsVisible()
 	return bIsVisible;
 }
 
+void FAndroidWebBrowserWindow::FloatingCloseButtonPressed()
+{
+	if (OnFloatingCloseButtonPressed().IsBound())
+	{
+		OnFloatingCloseButtonPressed().Execute();
+	}
+}
+
 #endif // USE_ANDROID_JNI
diff --git a/Engine/Source/Runtime/WebBrowser/Private/Android/AndroidWebBrowserWindow.h b/Engine/Source/Runtime/WebBrowser/Private/Android/AndroidWebBrowserWindow.h
--- a/Engine/Source/Runtime/WebBrowser/Private/Android/AndroidWebBrowserWindow.h
+++ b/Engine/Source/Runtime/WebBrowser/Private/Android/AndroidWebBrowserWindow.h
@@ -93,6 +93,7 @@ public:
 	virtual void SetIsDisabled(bool bValue) override;
 	virtual TSharedPtr<SWindow> GetParentWindow() const override;
 	virtual void SetParentWindow(TSharedPtr<SWindow> Window) override;
+	virtual void ShowFloatingCloseButton(bool bShow, bool bDraggable) override;
 
 	DECLARE_DERIVED_EVENT(FAndroidWebBrowserWindow, IWebBrowserWindow::FOnDocumentStateChanged, FOnDocumentStateChanged);
 	virtual FOnDocumentStateChanged& OnDocumentStateChanged() override
@@ -144,6 +145,11 @@ public:
 		return CloseWindowDelegate;
 	}
 
+	virtual FOnFloatingCloseButtonPressed& OnFloatingCloseButtonPressed() override
+	{
+		return FloatingCloseButtonPressedDelegate;
+	}
+
 	virtual FCursorReply OnCursorQuery(const FGeometry& MyGeometry, const FPointerEvent& CursorEvent) override
 	{
 		return FCursorReply::Unhandled();
@@ -233,6 +239,10 @@ public:
 	*/
 	bool IsVisible();
 
+	/**
+	* Webvew floating close button was pressed
+	*/
+	void FloatingCloseButtonPressed();
 private:
 
 	void SetTitle(const FString& InTitle)
@@ -304,6 +314,9 @@ private:
 	/** Delegate for handling requests to close new windows that were created. */
 	FOnCloseWindow CloseWindowDelegate;
 
+	/** Delegate for handling requests to close from the webview floating close button. */
+	FOnFloatingCloseButtonPressed FloatingCloseButtonPressedDelegate;
+
 	/** Delegate for handling requests to show the popup menu. */
 	FOnShowPopup ShowPopupEvent;
 
diff --git a/Engine/Source/Runtime/WebBrowser/Private/CEF/CEFWebBrowserWindow.h b/Engine/Source/Runtime/WebBrowser/Private/CEF/CEFWebBrowserWindow.h
--- a/Engine/Source/Runtime/WebBrowser/Private/CEF/CEFWebBrowserWindow.h
+++ b/Engine/Source/Runtime/WebBrowser/Private/CEF/CEFWebBrowserWindow.h
@@ -159,6 +159,9 @@ public:
 	virtual void SetIsDisabled(bool bValue) override;
 	virtual TSharedPtr<SWindow> GetParentWindow() const override;
 	virtual void SetParentWindow(TSharedPtr<SWindow> Window) override;
+	virtual void ShowFloatingCloseButton(bool bShow, bool bDraggable) override
+	{		
+	}
 
 	DECLARE_DERIVED_EVENT(FCEFWebBrowserWindow, IWebBrowserWindow::FOnDocumentStateChanged, FOnDocumentStateChanged);
 	virtual FOnDocumentStateChanged& OnDocumentStateChanged() override
@@ -210,6 +213,11 @@ public:
 		return CloseWindowDelegate;
 	}
 
+	virtual FOnFloatingCloseButtonPressed& OnFloatingCloseButtonPressed() override
+	{
+		return FloatingCloseButtonPressedDelegate;
+	}
+
 	virtual FCursorReply OnCursorQuery( const FGeometry& MyGeometry, const FPointerEvent& CursorEvent ) override
 	{
 		return Cursor == EMouseCursor::Default ? FCursorReply::Unhandled() : FCursorReply::Cursor(Cursor);
@@ -689,6 +697,9 @@ private:
 	/** Delegate for handling requests to close new windows that were created. */
 	FOnCloseWindow CloseWindowDelegate;
 
+	/** Delegate for handling requests to close from the webview floating close button. */
+	FOnFloatingCloseButtonPressed FloatingCloseButtonPressedDelegate;
+
 	/** Delegate for handling resource load requests */
 	FOnBeforeResourceLoadDelegate BeforeResourceLoadDelegate;
 
diff --git a/Engine/Source/Runtime/WebBrowser/Private/IOS/IOSPlatformWebBrowser.cpp b/Engine/Source/Runtime/WebBrowser/Private/IOS/IOSPlatformWebBrowser.cpp
--- a/Engine/Source/Runtime/WebBrowser/Private/IOS/IOSPlatformWebBrowser.cpp
+++ b/Engine/Source/Runtime/WebBrowser/Private/IOS/IOSPlatformWebBrowser.cpp
@@ -14,6 +14,7 @@
 #import <MetalKit/MetalKit.h>
 #include "ExternalTexture.h"
 #include "WebBrowserModule.h"
+#include "WebViewCloseButton.h"
 #include "IWebBrowserSingleton.h"
 
 #if !UE_BUILD_SHIPPING
@@ -39,7 +40,7 @@ class SIOSWebBrowserWidget : public SLeafWidget
 
 	SLATE_END_ARGS()
 
-		SIOSWebBrowserWidget()
+	SIOSWebBrowserWidget()
 		: WebViewWrapper(nil)
 	{}
 
@@ -53,13 +54,29 @@ class SIOSWebBrowserWidget : public SLeafWidget
 		// At this point we MUST be a Metal renderer.
 		check(bSupportsMetal);
 
-		WebViewWrapper = [IOSWebViewWrapper alloc];
-		[WebViewWrapper create : TSharedPtr<SIOSWebBrowserWidget>(this) userAgentApplication : Args._UserAgentApplication.GetNSString() useTransparency : Args._UseTransparency supportsMetal : bSupportsMetal supportsMetalMRT : bSupportsMetalMRT];
-
 		WebBrowserWindowPtr = Args._WebBrowserWindow;
 		IsIOS3DBrowser = false;
 
+		bool bEnableFloatingCloseButton = false;
+		GConfig->GetBool(TEXT("Browser"), TEXT("bEnableFloatingCloseButton"), bEnableFloatingCloseButton, GEngineIni);
+
+		TSharedPtr<SIOSWebBrowserWidget> SharedThis(this);
+		WebViewWrapper = [IOSWebViewWrapper alloc];
+		[WebViewWrapper create : SharedThis userAgentApplication : Args._UserAgentApplication.GetNSString() useTransparency : Args._UseTransparency supportsMetal : bSupportsMetal supportsMetalMRT : bSupportsMetalMRT enableFloatingCloseButton : bEnableFloatingCloseButton];
+
 #if !PLATFORM_TVOS
+		OnSafeFrameChangedEventHandle = FCoreDelegates::OnSafeFrameChangedEvent.AddLambda( [this, WeakThis = TWeakPtr<SIOSWebBrowserWidget>(SharedThis)]()
+										{
+											if (TSharedPtr<SIOSWebBrowserWidget> StrongThis = WeakThis.Pin())
+											{
+												if (WebViewWrapper != nil)
+												{
+													[WebViewWrapper didRotate];
+												}
+											}
+										});
+
+
 		TextureSamplePool = new FWebBrowserTextureSamplePool();
 		WebBrowserTextureSamplesQueue = MakeShared<FWebBrowserTextureSampleQueue, ESPMode::ThreadSafe>();
 		WebBrowserTexture = nullptr;
@@ -267,6 +284,8 @@ class SIOSWebBrowserWidget : public SLeafWidget
 
 	void Close()
 	{
+		FCoreDelegates::OnSafeFrameChangedEvent.Remove(OnSafeFrameChangedEventHandle);
+
 		if (WebViewWrapper != nil)
 		{
 			[WebViewWrapper close];
@@ -275,6 +294,14 @@ class SIOSWebBrowserWidget : public SLeafWidget
 		WebBrowserWindowPtr.Reset();
 	}
 
+	void ShowFloatingCloseButton(bool bShow, bool bDraggable)
+	{
+		if (WebViewWrapper != nil)
+		{
+			[WebViewWrapper showFloatingCloseButton: bShow setDraggable: bDraggable];
+		}
+	}
+	
 	void GoBack()
 	{
 		if (WebViewWrapper != nil)
@@ -439,7 +466,23 @@ class SIOSWebBrowserWidget : public SLeafWidget
 		}
 	}
 
-
+	void FloatingCloseButtonPressed()
+	{
+		if (WebBrowserWindowPtr.IsValid())
+		{
+			TWeakPtr<FWebBrowserWindow> AsyncWebBrowserWindowPtr = WebBrowserWindowPtr;
+			
+			[FIOSAsyncTask CreateTaskWithBlock : ^ bool(void)
+			{
+				if (TSharedPtr<FWebBrowserWindow> BrowserWindow = AsyncWebBrowserWindowPtr.Pin())
+				{
+					BrowserWindow->FloatingCloseButtonPressed();
+				}
+				return true;
+			}];
+		}
+	}
+	
 	~SIOSWebBrowserWidget()
 	{
 		Close();
@@ -468,20 +511,27 @@ private:
 
 	/** Texture sample object pool. */
 	FWebBrowserTextureSamplePool* TextureSamplePool;
+
+	// Handle to detect changes in valid area to move the floating close button */
+	FDelegateHandle OnSafeFrameChangedEventHandle;
 #endif
 };
 
+
+
+
 @implementation IOSWebViewWrapper
 
 #if !PLATFORM_TVOS
 @synthesize WebView;
+@synthesize CloseButton;
 @synthesize WebViewContainer;
 #endif
 @synthesize NextURL;
 @synthesize NextContent;
 
 -(void)create:(TSharedPtr<SIOSWebBrowserWidget>)InWebBrowserWidget userAgentApplication: (NSString*)UserAgentApplication useTransparency : (bool)InUseTransparency
-supportsMetal : (bool)InSupportsMetal supportsMetalMRT : (bool)InSupportsMetalMRT;
+supportsMetal : (bool)InSupportsMetal supportsMetalMRT : (bool)InSupportsMetalMRT enableFloatingCloseButton : (bool)bEnableFloatingCloseButton;
 {
 	WebBrowserWidget = InWebBrowserWidget;
 	NextURL = nil;
@@ -525,12 +575,36 @@ supportsMetal : (bool)InSupportsMetal supportsMetalMRT : (bool)InSupportsMetalMR
 			[self.WebView setOpaque : YES];
 		}
 
+		if (bEnableFloatingCloseButton)
+		{
+			CloseButton = MakeCloseButton();
+			[self.WebView addSubview : CloseButton];
+
+			TWeakPtr<SIOSWebBrowserWidget> WeakWebBrowserWidget = InWebBrowserWidget;
+			CloseButton.TapHandler = [^{
+				if (TSharedPtr<SIOSWebBrowserWidget> StrongWebBrowserWidget = WeakWebBrowserWidget.Pin())
+				{
+					StrongWebBrowserWidget->FloatingCloseButtonPressed();
+				}
+			} copy];
+		}
+
 		[theConfiguration release];
 		[self setDefaultVisibility];
 	});
 #endif
 }
 
+-(void)didRotate;
+{
+	dispatch_async(dispatch_get_main_queue(), ^ {
+		if (CloseButton)
+		{
+			[CloseButton setupLayout];
+		}
+	});
+}
+
 -(void)close;
 {
 #if !PLATFORM_TVOS
@@ -540,15 +614,35 @@ supportsMetal : (bool)InSupportsMetal supportsMetalMRT : (bool)InSupportsMetalMR
 		[self.WebViewContainer removeFromSuperview];
 		[self.WebView removeFromSuperview];
 		
+		if (CloseButton)
+		{
+			[CloseButton release];
+		}
 		[WebView release];
 		[WebViewContainer release];
 
+		CloseButton = nil;
 		WebView = nil;
 		WebViewContainer = nil;
 	});
 #endif
 }
 
+-(void)showFloatingCloseButton:(BOOL)bShow setDraggable:(BOOL)bDraggable;
+{
+	dispatch_async(dispatch_get_main_queue(), ^
+	{
+		if (CloseButton)
+		{
+			[CloseButton showButton:bShow setDraggable: bDraggable];
+		}
+		else
+		{
+			UE_LOG(LogIOS, Warning, TEXT("[PlatformWebBrowser]: Close button not enabled in config. Add config for Engine:[Browser]bEnableFloatingCloseButton=true"));
+		}
+	});
+}
+
 -(void)dealloc;
 {
 #if !PLATFORM_TVOS
@@ -559,6 +653,12 @@ supportsMetal : (bool)InSupportsMetal supportsMetalMRT : (bool)InSupportsMetalMR
 		WebView = nil;
 	};
 
+	if (CloseButton)
+	{
+		[CloseButton release];
+		CloseButton = nil;
+	};
+
 	[WebViewContainer release];
 	WebViewContainer = nil;
 #endif
@@ -585,6 +685,10 @@ supportsMetal : (bool)InSupportsMetal supportsMetalMRT : (bool)InSupportsMetalMR
 			{
 				bNeedsAddToView = false;
 				[[IOSAppDelegate GetDelegate].IOSView addSubview : WebViewContainer];
+				if (CloseButton != nil)
+				{
+					[CloseButton setupLayout];
+				}
 			}
 			else
 			{
@@ -1168,6 +1272,14 @@ void FWebBrowserWindow::SetParentWindow(TSharedPtr<SWindow> Window)
 	ParentWindow = Window;
 }
 
+void FWebBrowserWindow::ShowFloatingCloseButton(bool bShow, bool bDraggable)
+{
+	if (BrowserWidget)
+	{
+		BrowserWidget->ShowFloatingCloseButton(bShow, bDraggable);
+	}
+}
+
 void FWebBrowserWindow::ExecuteJavascript(const FString& Script)
 {
 	BrowserWidget->ExecuteJavascript(Script);
@@ -1220,4 +1332,11 @@ bool FWebBrowserWindow::IsVisible()
 	return bIsVisible;
 }
 
+void FWebBrowserWindow::FloatingCloseButtonPressed()
+{
+	if (OnFloatingCloseButtonPressed().IsBound())
+	{
+		OnFloatingCloseButtonPressed().Execute();
+	}
+}
 #endif
diff --git a/Engine/Source/Runtime/WebBrowser/Private/IOS/IOSPlatformWebBrowser.h b/Engine/Source/Runtime/WebBrowser/Private/IOS/IOSPlatformWebBrowser.h
--- a/Engine/Source/Runtime/WebBrowser/Private/IOS/IOSPlatformWebBrowser.h
+++ b/Engine/Source/Runtime/WebBrowser/Private/IOS/IOSPlatformWebBrowser.h
@@ -27,6 +27,8 @@ class SWebBrowserView;
 * Wrapper to contain the UIWebView and implement its delegate functions
 */
 #if !PLATFORM_TVOS
+@class WebViewCloseButton;
+
 @interface IOSWebViewWrapper : NSObject <WKUIDelegate, WKNavigationDelegate, WKScriptMessageHandler>
 #else
 @interface IOSWebViewWrapper : NSObject
@@ -42,6 +44,7 @@ class SWebBrowserView;
 }
 #if !PLATFORM_TVOS
 @property(strong) WKWebView* WebView;
+@property(strong) WebViewCloseButton* CloseButton;
 @property(strong) UIView* WebViewContainer;
 #endif
 @property(copy) NSURL* NextURL;
@@ -49,8 +52,10 @@ class SWebBrowserView;
 @property CGRect DesiredFrame;
 
 -(void)create:(TSharedPtr<SIOSWebBrowserWidget>)InWebBrowserWidget userAgentApplication: (NSString*)UserAgentApplication useTransparency : (bool)InUseTransparency
-supportsMetal : (bool)InSupportsMetal supportsMetalMRT : (bool)InSupportsMetalMRT;
+supportsMetal : (bool)InSupportsMetal supportsMetalMRT : (bool)InSupportsMetalMRT enableFloatingCloseButton : (bool)bEnableFloatingCloseButton;
+-(void)didRotate;
 -(void)close;
+-(void)showFloatingCloseButton:(BOOL)bShow setDraggable:(BOOL)bDraggable;
 -(void)dealloc;
 -(void)updateframe:(CGRect)InFrame;
 -(void)loadstring:(NSString*)InString dummyurl:(NSURL*)InURL;
@@ -153,6 +158,7 @@ public:
 	virtual void SetIsDisabled(bool bValue) override;
 	virtual TSharedPtr<SWindow> GetParentWindow() const override;
 	virtual void SetParentWindow(TSharedPtr<SWindow> Window) override;
+	virtual void ShowFloatingCloseButton(bool bShow, bool bDraggable) override;
 
 	// TODO: None of these events are actually called
 
@@ -206,6 +212,11 @@ public:
 		return CloseWindowDelegate;
 	}
 
+	virtual FOnFloatingCloseButtonPressed& OnFloatingCloseButtonPressed() override
+	{
+		return FloatingCloseButtonPressedDelegate;
+	}
+
 	virtual FCursorReply OnCursorQuery(const FGeometry& MyGeometry, const FPointerEvent& CursorEvent) override
 	{
 		//return Cursor == EMouseCursor::Default ? FCursorReply::Unhandled() : FCursorReply::Cursor(Cursor);
@@ -302,6 +313,12 @@ public:
 	*/
 	bool IsVisible();
 
+	/**
+	* Webvew floating close button was pressed
+	*/
+	void FloatingCloseButtonPressed();
+
+
 	void SetTitle(const FString& InTitle)
 	{
 		Title = InTitle;
@@ -368,6 +385,9 @@ private:
 	/** Delegate for handling requests to close new windows that were created. */
 	FOnCloseWindow CloseWindowDelegate;
 
+	/** Delegate for handling requests to close from the webview floating close button. */
+	FOnFloatingCloseButtonPressed FloatingCloseButtonPressedDelegate;
+
 	/** Delegate for handling requests to show the popup menu. */
 	FOnShowPopup ShowPopupEvent;
 
diff --git a/Engine/Source/Runtime/WebBrowser/Private/IOS/WebViewCloseButton.cpp b/Engine/Source/Runtime/WebBrowser/Private/IOS/WebViewCloseButton.cpp
new file mode 100644
--- /dev/null
+++ b/Engine/Source/Runtime/WebBrowser/Private/IOS/WebViewCloseButton.cpp
@@ -0,0 +1,397 @@
+// Copyright Epic Games, Inc. All Rights Reserved.
+#include "WebViewCloseButton.h"
+
+#if PLATFORM_IOS
+static constexpr float GFullscreenCloseButtonSize = 40.0f;
+static constexpr float GFullscreenCloseButtonInnerXInnerSize = 32.0f;
+static constexpr float GFullscreenCloseButtonInnerXOuterSize = 34.0f;
+static constexpr float GFullscreenCloseButtonOffsetFromEdge = 4.0f;
+static constexpr float GCloseButtonClickDragTolerance = 10.0f;
+
+static const uint8_t XTextureData[] =
+{
+		0x00, 0x00, 0x00, 0x00, 0x00, 0x00, 0x00, 0x00, 0x00, 0x00, 0x00, 0x00, 0x00, 0x00, 0x00, 0x00,
+		0x00, 0x00, 0x00, 0x00, 0x00, 0x00, 0x00, 0x00, 0x00, 0x00, 0x00, 0x00, 0x00, 0x00, 0x00, 0x00,
+		0x00, 0x00, 0x00, 0x00, 0x00, 0x00, 0x00, 0x00, 0x00, 0x00, 0x00, 0x00, 0x00, 0x00, 0x00, 0x00,
+		0x00, 0x00, 0x00, 0x00, 0x00, 0x00, 0x00, 0x00, 0x00, 0x00, 0x00, 0x00, 0x00, 0x00, 0x00, 0x00,
+		0x00, 0x00, 0x00, 0x00, 0x00, 0x00, 0x00, 0x00, 0x00, 0x00, 0x00, 0x00, 0x00, 0x00, 0x00, 0x00,
+		0x00, 0x00, 0x00, 0x00, 0x00, 0x00, 0x00, 0x00, 0x00, 0x00, 0x00, 0x00, 0x00, 0x00, 0x00, 0x00,
+		0x00, 0x00, 0x00, 0x00, 0x00, 0x00, 0x00, 0x71, 0x93, 0x3d, 0x00, 0x00, 0x00, 0x00, 0x00, 0x00,
+		0x00, 0x00, 0x00, 0x00, 0x00, 0x00, 0x50, 0xe9, 0xfa, 0xe3, 0x58, 0x00, 0x00, 0x00, 0x00, 0x00,
+		0x00, 0x00, 0x00, 0x00, 0x00, 0x00, 0x00, 0xbc, 0xf8, 0xfa, 0xe3, 0x58, 0x00, 0x00, 0x00, 0x00,
+		0x00, 0x00, 0x00, 0x00, 0x00, 0x00, 0x00, 0x1e, 0xc2, 0xf8, 0xfa, 0xe3, 0x58, 0x00, 0x00, 0x00,
+		0x00, 0x00, 0x00, 0x00, 0x00, 0x00, 0x00, 0x00, 0x1e, 0xc2, 0xf8, 0xfa, 0xe3, 0x58, 0x00, 0x00,
+		0x00, 0x00, 0x00, 0x00, 0x00, 0x00, 0x00, 0x00, 0x00, 0x1e, 0xc2, 0xf8, 0xfa, 0xe3, 0x58, 0x00,
+		0x00, 0x00, 0x00, 0x00, 0x00, 0x00, 0x00, 0x00, 0x00, 0x00, 0x1e, 0xc2, 0xf8, 0xfa, 0xe3, 0x58,
+		0x00, 0x00, 0x00, 0x00, 0x00, 0x00, 0x00, 0x00, 0x00, 0x00, 0x00, 0x1e, 0xc2, 0xf8, 0xfa, 0xe5,
+		0x00, 0x00, 0x00, 0x00, 0x00, 0x00, 0x00, 0x00, 0x00, 0x00, 0x00, 0x00, 0x1e, 0xc2, 0xf8, 0xfa,
+		0x00, 0x00, 0x00, 0x00, 0x00, 0x00, 0x00, 0x00, 0x00, 0x00, 0x00, 0x00, 0x00, 0x34, 0xdc, 0xfa,
+};
+
+struct FDefaultSymmetricTextureColorFn
+{
+	uint32_t operator()(uint32_t In) { return In; }
+};
+
+static UIImage* CreateSymmetricTexture(const uint8_t Data[16 * 16], TFunction<uint32_t(uint32_t)>&& ColorFn = FDefaultSymmetricTextureColorFn())
+{
+	uint32_t TextureData[32 * 32];
+
+	for (int Row = 0; Row < 16; ++Row)
+	{
+		for (int Col = 0; Col < 16; ++Col)
+		{
+			const uint32_t Gray = Data[Row * 16 + Col];
+			const uint32_t Color = ColorFn(0xFFFFFF | (Gray << 24));
+
+#define Q1 Row*32+Col
+#define Q2 Row*32+(32 - Col - 1)
+#define Q3 (32 - Row - 1)*32+Col
+#define Q4 (32 - Row - 1)*32+(32 - Col - 1)
+			TextureData[Q1] = Color;
+			TextureData[Q2] = Color;
+			TextureData[Q3] = Color;
+			TextureData[Q4] = Color;
+#undef Q1
+#undef Q2
+#undef Q3
+#undef Q4
+		}
+	}
+
+	UIGraphicsBeginImageContext(CGSizeMake(32, 32));
+	CGColorSpaceRef ColorSpace = CGColorSpaceCreateDeviceRGB();
+	CGBitmapInfo BitmapInfo = kCGImageAlphaFirst | kCGBitmapByteOrder32Little;
+	CGDataProviderRef DataProvider = CGDataProviderCreateWithData(NULL, TextureData, sizeof(TextureData), NULL);
+	CGImageRef ImageRef = CGImageCreate(32, 32, 8, 32, 32 * 4, ColorSpace, BitmapInfo, DataProvider, NULL, NO, kCGRenderingIntentDefault);
+	CGColorSpaceRelease(ColorSpace);
+	CGDataProviderRelease(DataProvider);
+
+	CGContextRef G = UIGraphicsGetCurrentContext();
+	CGContextClearRect(G, CGRectMake(0, 0, 32, 32));
+	CGContextDrawImage(G, CGRectMake(0, 0, 32, 32), ImageRef);
+	UIImage* Result = UIGraphicsGetImageFromCurrentImageContext();
+	UIGraphicsEndImageContext();
+
+	return Result;
+}
+
+static UIImage* CreateTintedImage(UIImage* SourceImage, UIColor* TintColor)
+{
+	bool bNeedsTint = false;
+	CGFloat Red, Green, Blue, Alpha;
+	if ([TintColor getRed:&Red green:&Green blue:&Blue alpha:&Alpha])
+	{
+		bNeedsTint = Red < 0.99 || Green < 0.99 || Blue < 0.99 || Alpha < 0.99;
+	}
+	else
+	{
+		bNeedsTint = !CGColorEqualToColor(UIColor.whiteColor.CGColor, TintColor.CGColor);
+	}
+
+	if (!bNeedsTint)
+	{
+		return SourceImage;
+	}
+
+	size_t ImageWidth = (size_t)(SourceImage.size.width * SourceImage.scale);
+	size_t ImageHeight = (size_t)(SourceImage.size.height * SourceImage.scale);
+	if (!ImageHeight || !ImageWidth)
+	{
+		return SourceImage;
+	}
+
+	uint32_t* PixelData = (uint32_t*)calloc(ImageWidth*ImageHeight, 4);
+	CGColorSpaceRef ColorSpace = CGColorSpaceCreateDeviceRGB();
+
+	CGRect Rect = CGRectMake(0, 0, ImageWidth, ImageHeight);
+	CGRect FirstPixelDataEntryRect = CGRectMake(0, ImageHeight-1, 1, 1); // bottom-most scanline is the first row in the array
+	CGContextRef BitmapContext = CGBitmapContextCreate(PixelData, ImageWidth, ImageHeight, 8, ImageWidth*4, ColorSpace, kCGImageAlphaPremultipliedLast | kCGImageByteOrder32Big);
+	CGContextClearRect(BitmapContext, Rect);
+	CGContextSetBlendMode(BitmapContext, kCGBlendModeNormal);
+	CGContextSetFillColorSpace(BitmapContext, CGColorGetColorSpace(TintColor.CGColor));
+	CGContextSetFillColor(BitmapContext, CGColorGetComponents(TintColor.CGColor));
+	CGContextFillRect(BitmapContext, FirstPixelDataEntryRect); // write a single pixel in the tinting color so we can read it out as RGBA
+	CGContextFlush(BitmapContext);
+
+	// PixelData now contains a single entry with our TintColor -- grab it.
+	const uint32_t PremultipliedTintColor = *PixelData;
+	uint32_t TA = (PremultipliedTintColor >> 24) & 0xFF;
+	uint32_t TB = (PremultipliedTintColor >> 16) & 0xFF;
+	uint32_t TG = (PremultipliedTintColor >> 8) & 0xFF;
+	uint32_t TR = PremultipliedTintColor & 0xFF;
+	if (TA)
+	{
+		TB = TB * 255 / TA + 1;
+		TG = TG * 255 / TA + 1;
+		TR = TR * 255 / TA + 1;
+	}
+
+	CGContextClearRect(BitmapContext, FirstPixelDataEntryRect);
+	CGContextDrawImage(BitmapContext, Rect, SourceImage.CGImage);
+	CGContextFlush(BitmapContext);
+
+	// PixelData now contains the SourceImage pixel data with premultiplied alpha.
+	for (size_t Y = 0; Y < ImageHeight; ++Y)
+	{
+		uint32_t* PixelRow = &PixelData[Y*ImageWidth];
+		for (size_t X = 0; X < ImageWidth; ++X)
+		{
+			uint32_t& Pixel = PixelRow[X];
+			uint32_t SA = (Pixel >> 24) & 0xFF, SB = (Pixel >> 16) & 0xFF, SG = (Pixel >> 8) & 0xFF, SR = Pixel & 0xFF;
+			if (SA)
+			{
+				SB = SB * 255 / SA;
+				SG = SG * 255 / SA;
+				SR = SR * 255 / SA;
+
+				uint32_t NewAlpha = SA*TA/255;
+				NewAlpha = NewAlpha + ((255 - NewAlpha)*TA / 255) + 1;
+
+				SB = (SB*TB*NewAlpha + 0x80) >> 16;
+				SG = (SG*TG*NewAlpha + 0x80) >> 16;
+				SR = (SR*TR*NewAlpha + 0x80) >> 16;
+				SA = NewAlpha - 1;
+				Pixel = (SA << 24) | (SB << 16) | (SG << 8) | SR;
+			}
+		}
+	}
+
+	// Get an image of the pixel data
+	CGImageRef TempImage = CGBitmapContextCreateImage(BitmapContext);
+	UIGraphicsBeginImageContext(CGSizeMake(ImageWidth, ImageHeight));
+	CGContextRef CopyImageContext = UIGraphicsGetCurrentContext();
+	CGContextDrawImage(CopyImageContext, Rect, TempImage);
+	UIImage* OutImage = UIGraphicsGetImageFromCurrentImageContext();
+	UIGraphicsEndImageContext();
+
+	// Release resources
+	CGImageRelease(TempImage);
+	CGContextRelease(BitmapContext);
+	CGColorSpaceRelease(ColorSpace);
+
+	return OutImage;
+}
+
+@interface WebViewCloseButtonDrawingLayer : NSObject
+@property (nonatomic, retain) UIImage* Texture;
+@property (nonatomic) CGFloat TextureInset;
+@property (nonatomic) BOOL bDrawWhileTouchDown;
+@property (nonatomic) CGBlendMode BlendMode;
+@end
+
+@implementation WebViewCloseButtonDrawingLayer
+
++ (instancetype)drawingLayerWithTexture:(UIImage*)Texture blendMode:(CGBlendMode)BlendMode inset:(CGFloat)TextureInset forTouchDown:(BOOL)bDrawWhileTouchDown
+{
+	WebViewCloseButtonDrawingLayer* Layer = [self new];
+	Layer.Texture = Texture;
+	Layer.TextureInset = TextureInset;
+	Layer.bDrawWhileTouchDown = bDrawWhileTouchDown;
+	Layer.BlendMode = BlendMode;
+	return Layer;
+}
+
+@end
+
+@implementation WebViewCloseButton
+{
+	NSMutableArray* DrawingLayers;
+	BOOL bIsDraggable;
+	BOOL bIsTouchDown;
+	BOOL bMayBeDragged;
+	CGPoint TouchDownPoint;
+	CGPoint DownDisplacement;
+	CGPoint Displacement;
+	CGPoint LayoutPosition;
+	CGRect SafeBounds;
+}
+
+- (void)dealloc
+{
+	[DrawingLayers release];
+	
+	[super dealloc];
+}
+
+- (instancetype)initWithFrame:(CGRect)frame
+{
+	if (self = [super initWithFrame:frame])
+	{
+		DrawingLayers = [[NSMutableArray arrayWithCapacity:4] retain];
+		self.hidden = YES;
+		bIsTouchDown = NO;
+		bMayBeDragged = NO;
+		bIsDraggable = NO;
+	}
+	return self;
+}
+
+-(void)showButton:(BOOL)bShow setDraggable:(BOOL)bDraggable
+{
+	self.hidden = !bShow;
+	if (bIsDraggable != bDraggable)
+	{
+		bIsDraggable = bDraggable;
+		[self setupLayout];
+	}
+}
+
+- (void)handleTap
+{
+	if (self.TapHandler)
+	{
+		self.TapHandler();
+	}
+}
+
+- (void)addDrawingLayerTexture:(UIImage *)Texture textureMultiply:(nonnull UIColor *)TextureMultiply blendMode:(CGBlendMode)BlendMode inset:(CGFloat)TextureInset forTouchDown:(BOOL)bInIsTouchDown
+{
+	if (Texture)
+	{
+		Texture = CreateTintedImage(Texture, TextureMultiply);
+		[DrawingLayers addObject:[WebViewCloseButtonDrawingLayer drawingLayerWithTexture:Texture blendMode:BlendMode inset:TextureInset forTouchDown:bInIsTouchDown]];
+		[self setNeedsDisplay];
+	}
+}
+
+- (void)setupLayout
+{
+	bMayBeDragged = NO;
+	Displacement = CGPointMake(0,0);
+	self.bounds = CGRectMake(0,0, GFullscreenCloseButtonSize, GFullscreenCloseButtonSize);
+	
+	if (!self.superview)
+	{
+		return;
+	}
+
+	UIEdgeInsets SuperInsets = [self.superview safeAreaInsets];
+	
+	// Substract insets from parent bounds
+	SafeBounds = self.superview.bounds;
+	SafeBounds.origin.x += SuperInsets.left;
+	SafeBounds.origin.y += SuperInsets.top;
+	SafeBounds.size.width -= SuperInsets.left + SuperInsets.right;
+	SafeBounds.size.height -= SuperInsets.top + SuperInsets.bottom;
+	
+	// Substract space for some margins and size of the brutton from its center
+	SafeBounds = CGRectInset(SafeBounds, GFullscreenCloseButtonSize/2 + GFullscreenCloseButtonOffsetFromEdge, GFullscreenCloseButtonSize/2 + GFullscreenCloseButtonOffsetFromEdge);
+
+	// Position the button in the top right corner fo the safe bounds
+	LayoutPosition = CGPointMake(SafeBounds.origin.x + SafeBounds.size.width, SafeBounds.origin.y);
+	self.center = LayoutPosition;
+}
+
+- (void)drawRect:(CGRect)rect
+{
+	for (WebViewCloseButtonDrawingLayer* Layer in DrawingLayers)
+	{
+		if (Layer.bDrawWhileTouchDown != bIsTouchDown)
+		{
+			continue;
+		}
+
+		CGRect RectCopy = CGRectInset(rect, Layer.TextureInset, Layer.TextureInset);
+		[Layer.Texture drawInRect:RectCopy blendMode:Layer.BlendMode alpha:1.0];
+	}
+}
+
+- (void)touchesBegan:(NSSet<UITouch *> *)touches withEvent:(UIEvent *)event
+{
+	if (touches.count > 1)
+	{
+		return;
+	}
+
+	bIsTouchDown = [self pointInside:[touches.anyObject locationInView:self] withEvent:event];
+	if (bIsTouchDown)
+	{
+		bMayBeDragged = bIsDraggable;
+		UITouch* Touch = [[touches allObjects] objectAtIndex: 0];
+		TouchDownPoint = [Touch locationInView: self.superview];
+		DownDisplacement = CGPointMake(TouchDownPoint.x - Displacement.x, TouchDownPoint.y - Displacement.y);
+	}
+	[self setNeedsDisplay];
+}
+
+- (void)touchesMoved:(NSSet<UITouch *> *)touches withEvent:(UIEvent *)event
+{
+	if (touches.count > 1)
+	{
+		return;
+	}
+
+	if (bMayBeDragged)
+	{
+		UITouch* Touch = [[touches allObjects] objectAtIndex: 0];
+		CGPoint Point = [Touch locationInView: self.superview];
+		CGPoint AttemptedDisplacement = CGPointMake(Point.x - DownDisplacement.x, Point.y - DownDisplacement.y);
+
+		CGPoint UpdatedCenter = CGPointMake( FMath::Clamp(LayoutPosition.x + AttemptedDisplacement.x, CGRectGetMinX(SafeBounds), CGRectGetMaxX(SafeBounds)),
+											 FMath::Clamp(LayoutPosition.y + AttemptedDisplacement.y, CGRectGetMinY(SafeBounds), CGRectGetMaxY(SafeBounds)));
+		
+		Displacement = CGPointMake(UpdatedCenter.x - LayoutPosition.x, UpdatedCenter.y - LayoutPosition.y);
+		
+		self.center = UpdatedCenter;
+	}
+	bIsTouchDown = [self pointInside:[touches.anyObject locationInView:self] withEvent:event];
+	[self setNeedsDisplay];
+}
+
+- (void)touchesEnded:(NSSet<UITouch *> *)touches withEvent:(UIEvent *)event
+{
+	if (touches.count > 1)
+	{
+		return;
+	}
+
+	bIsTouchDown = NO;
+	[self setNeedsDisplay];
+
+	UITouch* Touch = [[touches allObjects] objectAtIndex: 0];
+	CGPoint Point = [Touch locationInView: self.superview];
+	CGPoint AbsDisplacement = CGPointMake(Point.x - TouchDownPoint.x, Point.y - TouchDownPoint.y);
+	
+	if (FMath::Abs(AbsDisplacement.x) < GCloseButtonClickDragTolerance && FMath::Abs(AbsDisplacement.y) < GCloseButtonClickDragTolerance )
+	{
+		if ([self pointInside:[touches.anyObject locationInView:self] withEvent:event])
+		{
+			[self handleTap];
+		}
+	}
+}
+
+- (void)touchesCancelled:(NSSet<UITouch *> *)touches withEvent:(UIEvent *)event
+{
+	bIsTouchDown = NO;
+	bMayBeDragged = NO;
+	[self setNeedsDisplay];
+}
+
+- (void)touchesEstimatedPropertiesUpdated:(NSSet<UITouch *> *)touches
+{
+}
+
+@end
+
+WebViewCloseButton* MakeCloseButton()
+{
+	WebViewCloseButton* Button = [[WebViewCloseButton alloc] initWithFrame:CGRectZero];
+	Button.translatesAutoresizingMaskIntoConstraints = NO;
+	Button.opaque = NO;
+
+	UIImage* XImage = CreateSymmetricTexture(XTextureData);
+
+	// X states
+	CGFloat InsetWhite = (GFullscreenCloseButtonSize - GFullscreenCloseButtonInnerXInnerSize) / 2.0;
+	CGFloat InsetDark = (GFullscreenCloseButtonSize - GFullscreenCloseButtonInnerXOuterSize) / 2.0;
+	[Button addDrawingLayerTexture:XImage textureMultiply:UIColor.blackColor blendMode:kCGBlendModeNormal inset:InsetDark forTouchDown:NO];
+	[Button addDrawingLayerTexture:XImage textureMultiply:UIColor.whiteColor blendMode:kCGBlendModeNormal inset:InsetWhite forTouchDown:NO];
+	[Button addDrawingLayerTexture:XImage textureMultiply:UIColor.whiteColor blendMode:kCGBlendModeNormal inset:InsetDark forTouchDown:YES];
+	[Button addDrawingLayerTexture:XImage textureMultiply:UIColor.blackColor blendMode:kCGBlendModeNormal inset:InsetWhite forTouchDown:YES];
+
+	return Button;
+}
+
+#endif
diff --git a/Engine/Source/Runtime/WebBrowser/Private/IOS/WebViewCloseButton.h b/Engine/Source/Runtime/WebBrowser/Private/IOS/WebViewCloseButton.h
new file mode 100644
--- /dev/null
+++ b/Engine/Source/Runtime/WebBrowser/Private/IOS/WebViewCloseButton.h
@@ -0,0 +1,23 @@
+// Copyright Epic Games, Inc. All Rights Reserved.
+
+#pragma once
+
+#include "CoreMinimal.h"
+
+#if PLATFORM_IOS
+
+#if !PLATFORM_TVOS
+#import <UIKit/UIKit.h>
+
+@interface WebViewCloseButton : UIView
+@property (nonatomic, copy) void (^TapHandler)(void);
+
+- (void)setupLayout;
+- (void)showButton:(BOOL)bShow setDraggable:(BOOL)bDraggable;
+
+@end
+
+WebViewCloseButton* MakeCloseButton();
+#endif
+
+#endif
diff --git a/Engine/Source/Runtime/WebBrowser/Private/Native/NativeWebBrowserProxy.h b/Engine/Source/Runtime/WebBrowser/Private/Native/NativeWebBrowserProxy.h
--- a/Engine/Source/Runtime/WebBrowser/Private/Native/NativeWebBrowserProxy.h
+++ b/Engine/Source/Runtime/WebBrowser/Private/Native/NativeWebBrowserProxy.h
@@ -70,6 +70,9 @@ public:
 	virtual void SetIsDisabled(bool bValue) override;
 	virtual TSharedPtr<SWindow> GetParentWindow() const override;
 	virtual void SetParentWindow(TSharedPtr<SWindow> Window) override;
+	virtual void ShowFloatingCloseButton(bool bShow, bool bDraggable) override
+	{		
+	}
 
 	// @todo: None of these are actually called at the moment.
 	DECLARE_DERIVED_EVENT(FNativeWebBrowserProxy, IWebBrowserWindow::FOnDocumentStateChanged, FOnDocumentStateChanged);
@@ -230,6 +233,9 @@ private:
 	/** Delegate for handling requests to close new windows that were created. */
 	FOnCloseWindow CloseWindowDelegate;
 
+	/** Delegate for handling requests to close from the webview floating close button. */
+	FOnFloatingCloseButtonPressed FloatingCloseButtonPressedDelegate;
+
 	/** Delegate for handling requests to show the popup menu. */
 	FOnShowPopup ShowPopupEvent;
 
diff --git a/Engine/Source/Runtime/WebBrowser/Private/SWebBrowserView.cpp b/Engine/Source/Runtime/WebBrowser/Private/SWebBrowserView.cpp
--- a/Engine/Source/Runtime/WebBrowser/Private/SWebBrowserView.cpp
+++ b/Engine/Source/Runtime/WebBrowser/Private/SWebBrowserView.cpp
@@ -93,6 +93,7 @@ void SWebBrowserView::Construct(const FArguments& InArgs, const TSharedPtr<IWebB
 	OnBeforePopup = InArgs._OnBeforePopup;
 	OnCreateWindow = InArgs._OnCreateWindow;
 	OnCloseWindow = InArgs._OnCloseWindow;
+	OnFloatingCloseButtonPressed = InArgs._OnFloatingCloseButtonPressed;
 	AddressBarUrl = FText::FromString(InArgs._InitialURL);
 	PopupMenuMethod = InArgs._PopupMenuMethod;
 	OnUnhandledKeyDown = InArgs._OnUnhandledKeyDown;
@@ -153,6 +154,11 @@ void SWebBrowserView::Construct(const FArguments& InArgs, const TSharedPtr<IWebB
 			BrowserWindow->OnCloseWindow().BindSP(this, &SWebBrowserView::HandleCloseWindow);
 		}
 
+		if (OnFloatingCloseButtonPressed.IsBound())
+		{
+			BrowserWindow->OnFloatingCloseButtonPressed().BindSP(this, &SWebBrowserView::HandleFloatingCloseButtonPressed);
+		}
+
 		BrowserWindow->OnDocumentStateChanged().AddSP(this, &SWebBrowserView::HandleBrowserWindowDocumentStateChanged);
 		BrowserWindow->OnNeedsRedraw().AddSP(this, &SWebBrowserView::HandleBrowserWindowNeedsRedraw);
 		BrowserWindow->OnTitleChanged().AddSP(this, &SWebBrowserView::HandleTitleChanged);
@@ -497,6 +503,11 @@ void SWebBrowserView::CloseBrowser()
 	BrowserWindow->CloseBrowser(true /*force*/, true /*block until closed*/);
 }
 
+void SWebBrowserView::ShowFloatingCloseButton(bool bShow, bool bDraggable)
+{
+	BrowserWindow->ShowFloatingCloseButton(bShow, bDraggable);
+}
+
 void SWebBrowserView::HandleToolTip(FString ToolTipText)
 {
 	if(ToolTipText.IsEmpty())
@@ -594,6 +605,14 @@ bool SWebBrowserView::HandleCloseWindow(const TWeakPtr<IWebBrowserWindow>& NewBr
 	return false;
 }
 
+void SWebBrowserView::HandleFloatingCloseButtonPressed()
+{
+	if (OnFloatingCloseButtonPressed.IsBound())
+	{
+		OnFloatingCloseButtonPressed.Execute();
+	}
+}
+
 void SWebBrowserView::BindUObject(const FString& Name, UObject* Object, bool bIsPermanent)
 {
 	if (BrowserWindow.IsValid())
diff --git a/Engine/Source/Runtime/WebBrowser/Public/IWebBrowserWindow.h b/Engine/Source/Runtime/WebBrowser/Public/IWebBrowserWindow.h
--- a/Engine/Source/Runtime/WebBrowser/Public/IWebBrowserWindow.h
+++ b/Engine/Source/Runtime/WebBrowser/Public/IWebBrowserWindow.h
@@ -368,7 +368,13 @@ public:
 	*/
 	virtual void SetParentWindow(TSharedPtr<class SWindow> Window) = 0;
 
+	/**
+	* Shows floating button in the browser if supported by platform
+	*/
+	virtual void ShowFloatingCloseButton(bool bShow, bool bDraggable) = 0;
+
 	virtual void CheckTickActivity() {};
+
 public:
 
 	/** A delegate that is invoked when the loading state of a document changed. */
@@ -424,6 +430,10 @@ public:
 	DECLARE_DELEGATE_RetVal_OneParam(bool, FOnCloseWindow, const TWeakPtr<IWebBrowserWindow>& /*BrowserWindow*/)
 	virtual FOnCloseWindow& OnCloseWindow() = 0;
 
+	/** A delegate that is invoked when close is requested from the floating close button. */
+	DECLARE_DELEGATE(FOnFloatingCloseButtonPressed)
+	virtual FOnFloatingCloseButtonPressed& OnFloatingCloseButtonPressed() = 0;
+
 	/** A delegate that is invoked when the browser needs to show a popup menu. */
 	DECLARE_EVENT_OneParam(IWebBrowserWindow, FOnShowPopup, const FIntRect& /*PopupSize*/)
 	virtual FOnShowPopup& OnShowPopup() = 0;
diff --git a/Engine/Source/Runtime/WebBrowser/Public/SWebBrowserView.h b/Engine/Source/Runtime/WebBrowser/Public/SWebBrowserView.h
--- a/Engine/Source/Runtime/WebBrowser/Public/SWebBrowserView.h
+++ b/Engine/Source/Runtime/WebBrowser/Public/SWebBrowserView.h
@@ -26,6 +26,7 @@ DECLARE_DELEGATE_RetVal_TwoParams(bool, FOnCreateWindowDelegate, const TWeakPtr<
 DECLARE_DELEGATE_RetVal_OneParam(bool, FOnCloseWindowDelegate, const TWeakPtr<IWebBrowserWindow>&);
 DECLARE_DELEGATE_RetVal_OneParam(TSharedPtr<IToolTip>, FOnCreateToolTip, const FText&);
 DECLARE_DELEGATE_FourParams(FOnConsoleMessageDelegate, const FString& /*Message*/, const FString& /*Source*/, int32 /*Line*/, EWebBrowserConsoleLogSeverity /*Severity*/);
+DECLARE_DELEGATE(FOnFloatingCloseButtonPressedDelegate);
 
 #if WITH_CEF3
 typedef SViewport SWebBrowserWidget;
@@ -124,6 +125,9 @@ public:
 		/** Called when a browser window close event is detected */
 		SLATE_EVENT(FOnCloseWindowDelegate, OnCloseWindow)
 
+		/** Called when a browser window requests to close through the floating close button*/
+		SLATE_EVENT(FOnFloatingCloseButtonPressedDelegate, OnFloatingCloseButtonPressed)
+
 		/** Called before browser navigation. */
 		SLATE_EVENT(FOnBeforeBrowse, OnBeforeNavigation)
 
@@ -289,6 +293,9 @@ public:
 		Calling this is optional, CEF has object lifetime requirements that mean on shutdown you must destroy browsers before exit.*/
 	WEBBROWSER_API void CloseBrowser();
 
+	/** Shows floating close button.*/
+	WEBBROWSER_API void ShowFloatingCloseButton(bool bShow, bool bDraggable);
+
 private:
 
 	WEBBROWSER_API void SetupParentWindowHandlers();
@@ -331,6 +338,9 @@ private:
 	 */
 	WEBBROWSER_API bool HandleCloseWindow(const TWeakPtr<IWebBrowserWindow>& BrowserWindow);
 
+	/** Callback for handling requests from floating close button */
+	WEBBROWSER_API void HandleFloatingCloseButtonPressed();
+
 	/** Callback for showing dialogs to the user */
 	WEBBROWSER_API EWebBrowserDialogEventResponse HandleShowDialog(const TWeakPtr<IWebBrowserDialog>& DialogParams);
 
@@ -417,6 +427,9 @@ private:
 	/** A delegate that is invoked when a window close event is detected */
 	FOnCloseWindowDelegate OnCloseWindow;
 
+	/** A delegate that is invoked when the floating close button is pressed */
+	FOnFloatingCloseButtonPressedDelegate OnFloatingCloseButtonPressed;
+
 	/** A delegate that is invoked prior to browser navigation */
 	FOnBeforeBrowse OnBeforeNavigation;
 
