Fix TVOS build issues

diff --git a/Engine/Source/Runtime/WebBrowser/Private/IOS/IOSPlatformWebBrowser.cpp b/Engine/Source/Runtime/WebBrowser/Private/IOS/IOSPlatformWebBrowser.cpp
--- a/Engine/Source/Runtime/WebBrowser/Private/IOS/IOSPlatformWebBrowser.cpp
+++ b/Engine/Source/Runtime/WebBrowser/Private/IOS/IOSPlatformWebBrowser.cpp
@@ -284,8 +284,9 @@ class SIOSWebBrowserWidget : public SLeafWidget
 
 	void Close()
 	{
+#if !PLATFORM_TVOS
 		FCoreDelegates::OnSafeFrameChangedEvent.Remove(OnSafeFrameChangedEventHandle);
-
+#endif
 		if (WebViewWrapper != nil)
 		{
 			[WebViewWrapper close];
@@ -597,12 +598,14 @@ supportsMetal : (bool)InSupportsMetal supportsMetalMRT : (bool)InSupportsMetalMR
 
 -(void)didRotate;
 {
+#if !PLATFORM_TVOS
 	dispatch_async(dispatch_get_main_queue(), ^ {
 		if (CloseButton)
 		{
 			[CloseButton setupLayout];
 		}
 	});
+#endif
 }
 
 -(void)close;
@@ -630,6 +633,7 @@ supportsMetal : (bool)InSupportsMetal supportsMetalMRT : (bool)InSupportsMetalMR
 
 -(void)showFloatingCloseButton:(BOOL)bShow setDraggable:(BOOL)bDraggable;
 {
+#if !PLATFORM_TVOS
 	dispatch_async(dispatch_get_main_queue(), ^
 	{
 		if (CloseButton)
@@ -641,6 +645,7 @@ supportsMetal : (bool)InSupportsMetal supportsMetalMRT : (bool)InSupportsMetalMR
 			UE_LOG(LogIOS, Warning, TEXT("[PlatformWebBrowser]: Close button not enabled in config. Add config for Engine:[Browser]bEnableFloatingCloseButton=true"));
 		}
 	});
+#endif
 }
 
 -(void)dealloc;
diff --git a/Engine/Source/Runtime/WebBrowser/Private/IOS/WebViewCloseButton.cpp b/Engine/Source/Runtime/WebBrowser/Private/IOS/WebViewCloseButton.cpp
--- a/Engine/Source/Runtime/WebBrowser/Private/IOS/WebViewCloseButton.cpp
+++ b/Engine/Source/Runtime/WebBrowser/Private/IOS/WebViewCloseButton.cpp
@@ -2,6 +2,8 @@
 #include "WebViewCloseButton.h"
 
 #if PLATFORM_IOS
+#if !PLATFORM_TVOS
+
 static constexpr float GFullscreenCloseButtonSize = 40.0f;
 static constexpr float GFullscreenCloseButtonInnerXInnerSize = 32.0f;
 static constexpr float GFullscreenCloseButtonInnerXOuterSize = 34.0f;
@@ -394,4 +396,5 @@ WebViewCloseButton* MakeCloseButton()
 	return Button;
 }
 
-#endif
+#endif // !PLATFORM_TVOS
+#endif // PLATFORM_IOS
