Add support to inject Javascript bound objects in onPageStarted instead on onPageLoaded

diff --git a/Engine/Source/Runtime/WebBrowser/Private/Android/AndroidWebBrowserWindow.cpp b/Engine/Source/Runtime/WebBrowser/Private/Android/AndroidWebBrowserWindow.cpp
--- a/Engine/Source/Runtime/WebBrowser/Private/Android/AndroidWebBrowserWindow.cpp
+++ b/Engine/Source/Runtime/WebBrowser/Private/Android/AndroidWebBrowserWindow.cpp
@@ -30,12 +30,14 @@ FAndroidWebBrowserWindow::FAndroidWebBrowserWindow(FString InUrl, TOptional<FStr
 	, bUseTransparency(InUseTransparency)
 	, DocumentState(EWebBrowserDocumentState::NoDocument)
 	, ErrorCode(0)
-	, Scripting(new FMobileJSScripting(bInJSBindingToLoweringEnabled))
 	, AndroidWindowSize(FIntPoint(1024, 768))
 	, bIsDisabled(false)
 	, bIsVisible(true)
 	, bTickedLastFrame(true)
 {
+	bool bInjectJSOnPageStarted = false;
+	GConfig->GetBool(TEXT("Browser"), TEXT("bInjectJSOnPageStarted"), bInjectJSOnPageStarted, GEngineIni);
+	Scripting = MakeShared<FMobileJSScripting>(bInJSBindingToLoweringEnabled, bInjectJSOnPageStarted);
 	// Deal with optional texture resolution override
 	FString WebViewTextureSize;
 	if (GConfig->GetString(TEXT("/Script/AndroidRuntimeSettings.AndroidRuntimeSettings"), TEXT("WebViewTextureSize"), WebViewTextureSize, GEngineIni))
@@ -307,19 +309,21 @@ void FAndroidWebBrowserWindow::NotifyDocumentError(const FString& InCurrentUrl,
 
 void FAndroidWebBrowserWindow::NotifyDocumentLoadingStateChange(const FString& InCurrentUrl, bool IsLoading)
 {
+	// This method may called several times with IsLoading = false (one time for each frame in the page)
 	// Ignore a load completed notification if there was an error.
 	// For load started, reset any errors from previous page load.
-	if (IsLoading || DocumentState != EWebBrowserDocumentState::Error)
+	bool bIsNotError = DocumentState != EWebBrowserDocumentState::Error;
+	if (IsLoading || bIsNotError)
 	{
-		if(!CurrentUrl.Equals(InCurrentUrl, ESearchCase::CaseSensitive))
+		if (!CurrentUrl.Equals(InCurrentUrl, ESearchCase::CaseSensitive))
 		{
 			CurrentUrl = InCurrentUrl;
 			UrlChangedEvent.Broadcast(CurrentUrl);
 		}
 
-		if(!IsLoading && !InCurrentUrl.StartsWith("javascript:"))
+		if (bIsNotError && !InCurrentUrl.StartsWith("javascript:"))
 		{
-			Scripting->PageLoaded(SharedThis(this));
+			Scripting->TryInjectJavascript(SharedThis(this), IsLoading);
 		}
 		ErrorCode = 0;
 		DocumentState = IsLoading
diff --git a/Engine/Source/Runtime/WebBrowser/Private/IOS/IOSPlatformWebBrowser.cpp b/Engine/Source/Runtime/WebBrowser/Private/IOS/IOSPlatformWebBrowser.cpp
--- a/Engine/Source/Runtime/WebBrowser/Private/IOS/IOSPlatformWebBrowser.cpp
+++ b/Engine/Source/Runtime/WebBrowser/Private/IOS/IOSPlatformWebBrowser.cpp
@@ -1023,12 +1023,14 @@ FWebBrowserWindow::FWebBrowserWindow(FString InUrl, TOptional<FString> InContent
 	, bUseTransparency(InUseTransparency)
 	, DocumentState(EWebBrowserDocumentState::NoDocument)
 	, ErrorCode(0)
-	, Scripting(new FMobileJSScripting(bInJSBindingToLoweringEnabled))
 	, IOSWindowSize(FIntPoint(500, 500))
 	, bIsDisabled(false)
 	, bIsVisible(true)
 	, bTickedLastFrame(true)
 {
+	bool bInjectJSOnPageStarted = false;
+	GConfig->GetBool(TEXT("Browser"), TEXT("bInjectJSOnPageStarted"), bInjectJSOnPageStarted, GEngineIni);
+	Scripting = MakeShared<FMobileJSScripting>(bInJSBindingToLoweringEnabled, bInjectJSOnPageStarted);
 }
 
 FWebBrowserWindow::~FWebBrowserWindow()
@@ -1241,7 +1243,8 @@ void FWebBrowserWindow::NotifyDocumentLoadingStateChange(const FString& InCurren
 {
 	// Ignore a load completed notification if there was an error.
 	// For load started, reset any errors from previous page load.
-	if (IsLoading || DocumentState != EWebBrowserDocumentState::Error)
+	bool bIsNotError = DocumentState != EWebBrowserDocumentState::Error;
+	if (IsLoading || bIsNotError)
 	{
 		if (!CurrentUrl.Equals(InCurrentUrl, ESearchCase::CaseSensitive))
 		{
@@ -1249,9 +1252,9 @@ void FWebBrowserWindow::NotifyDocumentLoadingStateChange(const FString& InCurren
 			UrlChangedEvent.Broadcast(CurrentUrl);
 		}
 
-		if (!IsLoading && !InCurrentUrl.StartsWith("javascript:"))
+		if (bIsNotError && !InCurrentUrl.StartsWith("javascript:"))
 		{
-			Scripting->PageLoaded(SharedThis(this));
+			Scripting->TryInjectJavascript(SharedThis(this), IsLoading);
 		}
 		ErrorCode = 0;
 		DocumentState = IsLoading
diff --git a/Engine/Source/Runtime/WebBrowser/Private/MobileJS/MobileJSScripting.cpp b/Engine/Source/Runtime/WebBrowser/Private/MobileJS/MobileJSScripting.cpp
--- a/Engine/Source/Runtime/WebBrowser/Private/MobileJS/MobileJSScripting.cpp
+++ b/Engine/Source/Runtime/WebBrowser/Private/MobileJS/MobileJSScripting.cpp
@@ -613,30 +613,35 @@ void FMobileJSScripting::InitializeScript(TSharedRef<class IWebBrowserWindow> In
 	InWindow->ExecuteJavascript(Script);
 }
 
-void FMobileJSScripting::PageLoaded(TSharedRef<class IWebBrowserWindow> InWindow)
+void FMobileJSScripting::TryInjectJavascript(TSharedRef<class IWebBrowserWindow> InWindow, bool bPageStarted)
 {
 	WindowPtr = InWindow;
 
-	// Expunge temporary objects.
-	for (TMap<TObjectPtr<UObject>, ObjectBinding>::TIterator It(BoundObjects); It; ++It)
+	// Only inject if page state matches requested one, either when page starts loading (bPageStarted = true) or when page completely loaded (bPageStarted = false)
+	if (bInjectJSOnPageStarted == bPageStarted)
 	{
-		if (!It->Value.bIsPermanent)
+		// Expunge temporary objects.
+		for (TMap<TObjectPtr<UObject>, ObjectBinding>::TIterator It(BoundObjects); It; ++It)
 		{
-			It.RemoveCurrent();
+			if (!It->Value.bIsPermanent)
+			{
+				It.RemoveCurrent();
+			}
 		}
-	}
 
-	InitializeScript(InWindow);
-	FString Script;
-	for (auto& Item : PermanentUObjectsByName)
-	{
-		Script.Append(*FString::Printf(TEXT("window.ue['%s'] = %s;"), *Item.Key.ReplaceCharWithEscapedChar(), *ConvertObject(Item.Value)));
+		InitializeScript(InWindow);
+		FString Script;
+		for (auto& Item : PermanentUObjectsByName)
+		{
+			Script.Append(*FString::Printf(TEXT("window.ue['%s'] = %s;"), *Item.Key.ReplaceCharWithEscapedChar(), *ConvertObject(Item.Value)));
+		}
+		InWindow->ExecuteJavascript(Script);
 	}
-	InWindow->ExecuteJavascript(Script);
 }
 
-FMobileJSScripting::FMobileJSScripting(bool bJSBindingToLoweringEnabled)
+FMobileJSScripting::FMobileJSScripting(bool bJSBindingToLoweringEnabled, bool bInjectJSOnPageStarted)
 	: FWebJSScripting(bJSBindingToLoweringEnabled)
+	, bInjectJSOnPageStarted(bInjectJSOnPageStarted)
 {
 }
 
diff --git a/Engine/Source/Runtime/WebBrowser/Private/MobileJS/MobileJSScripting.h b/Engine/Source/Runtime/WebBrowser/Private/MobileJS/MobileJSScripting.h
--- a/Engine/Source/Runtime/WebBrowser/Private/MobileJS/MobileJSScripting.h
+++ b/Engine/Source/Runtime/WebBrowser/Private/MobileJS/MobileJSScripting.h
@@ -22,7 +22,7 @@ public:
 	static const FString JSMessageTag;
 	static const FString JSMessageHandler;
 
-	FMobileJSScripting(bool bJSBindingToLoweringEnabled);
+	FMobileJSScripting(bool bJSBindingToLoweringEnabled, bool bInjectJSOnPageStarted);
 
 	virtual void BindUObject(const FString& Name, UObject* Object, bool bIsPermanent = true) override;
 	virtual void UnbindUObject(const FString& Name, UObject* Object = nullptr, bool bIsPermanent = true) override;
@@ -43,7 +43,7 @@ public:
 
 	virtual void InvokeJSFunction(FGuid FunctionId, int32 ArgCount, FWebJSParam Arguments[], bool bIsError=false) override;
 	virtual void InvokeJSErrorResult(FGuid FunctionId, const FString& Error) override;
-	void PageLoaded(TSharedRef<class IWebBrowserWindow> InWindow); // Called on page load
+	void TryInjectJavascript(TSharedRef<class IWebBrowserWindow> InWindow, bool bPageStarted);
 
 	void SetWindow(TSharedRef<class IWebBrowserWindow> InWindow);
 
@@ -63,6 +63,9 @@ private:
 
 	/** Pointer to the Mobile Browser for this window. */
 	TWeakPtr<class IWebBrowserWindow> WindowPtr;
+	
+	/** When to inject JS code */
+	bool bInjectJSOnPageStarted;
 };
 
-#endif // PLATFORM_ANDROID  || PLATFORM_IOS
\ No newline at end of file
+#endif // PLATFORM_ANDROID  || PLATFORM_IOS
