Improve readability comments after changes to allow injecting JS code on page load start
Also passed some args as const ref

diff --git a/Engine/Source/Runtime/WebBrowser/Private/Android/AndroidWebBrowserWindow.cpp b/Engine/Source/Runtime/WebBrowser/Private/Android/AndroidWebBrowserWindow.cpp
--- a/Engine/Source/Runtime/WebBrowser/Private/Android/AndroidWebBrowserWindow.cpp
+++ b/Engine/Source/Runtime/WebBrowser/Private/Android/AndroidWebBrowserWindow.cpp
@@ -323,7 +323,14 @@ void FAndroidWebBrowserWindow::NotifyDocumentLoadingStateChange(const FString& I
 
 		if (bIsNotError && !InCurrentUrl.StartsWith("javascript:"))
 		{
-			Scripting->TryInjectJavascript(SharedThis(this), IsLoading);
+			if (IsLoading)
+			{
+				Scripting->PageStarted(SharedThis(this));
+			}
+			else
+			{
+				Scripting->PageLoaded(SharedThis(this));
+			}
 		}
 		ErrorCode = 0;
 		DocumentState = IsLoading
diff --git a/Engine/Source/Runtime/WebBrowser/Private/IOS/IOSPlatformWebBrowser.cpp b/Engine/Source/Runtime/WebBrowser/Private/IOS/IOSPlatformWebBrowser.cpp
--- a/Engine/Source/Runtime/WebBrowser/Private/IOS/IOSPlatformWebBrowser.cpp
+++ b/Engine/Source/Runtime/WebBrowser/Private/IOS/IOSPlatformWebBrowser.cpp
@@ -1254,7 +1254,14 @@ void FWebBrowserWindow::NotifyDocumentLoadingStateChange(const FString& InCurren
 
 		if (bIsNotError && !InCurrentUrl.StartsWith("javascript:"))
 		{
-			Scripting->TryInjectJavascript(SharedThis(this), IsLoading);
+			if (IsLoading)
+			{
+				Scripting->PageStarted(SharedThis(this));
+			}
+			else
+			{
+				Scripting->PageLoaded(SharedThis(this));
+			}
 		}
 		ErrorCode = 0;
 		DocumentState = IsLoading
diff --git a/Engine/Source/Runtime/WebBrowser/Private/MobileJS/MobileJSScripting.cpp b/Engine/Source/Runtime/WebBrowser/Private/MobileJS/MobileJSScripting.cpp
--- a/Engine/Source/Runtime/WebBrowser/Private/MobileJS/MobileJSScripting.cpp
+++ b/Engine/Source/Runtime/WebBrowser/Private/MobileJS/MobileJSScripting.cpp
@@ -311,7 +311,7 @@ void FMobileJSScripting::UnbindUObject(const FString& Name, UObject* Object, boo
 	}
 }
 
-void FMobileJSScripting::BindUObject(TSharedRef<class IWebBrowserWindow> InWindow, const FString& Name, UObject* Object, bool bIsPermanent)
+void FMobileJSScripting::BindUObject(const TSharedRef<class IWebBrowserWindow>& InWindow, const FString& Name, UObject* Object, bool bIsPermanent)
 {
 	WindowPtr = InWindow;
 
@@ -327,7 +327,7 @@ void FMobileJSScripting::BindUObject(TSharedRef<class IWebBrowserWindow> InWindo
 	InWindow->ExecuteJavascript(SetValueScript);
 }
 
-void FMobileJSScripting::UnbindUObject(TSharedRef<class IWebBrowserWindow> InWindow, const FString& Name, UObject* Object, bool bIsPermanent)
+void FMobileJSScripting::UnbindUObject(const TSharedRef<class IWebBrowserWindow>& InWindow, const FString& Name, UObject* Object, bool bIsPermanent)
 {
 	WindowPtr = InWindow;
 
@@ -579,7 +579,7 @@ bool FMobileJSScripting::HandleExecuteUObjectMethodMessage(const TArray<FString>
 	return true;
 }
 
-void FMobileJSScripting::InitializeScript(TSharedRef<class IWebBrowserWindow> InWindow)
+void FMobileJSScripting::InitializeScript(const TSharedRef<class IWebBrowserWindow>& InWindow)
 {
 	WindowPtr = InWindow;
 
@@ -613,30 +613,44 @@ void FMobileJSScripting::InitializeScript(TSharedRef<class IWebBrowserWindow> In
 	InWindow->ExecuteJavascript(Script);
 }
 
-void FMobileJSScripting::TryInjectJavascript(TSharedRef<class IWebBrowserWindow> InWindow, bool bPageStarted)
+void FMobileJSScripting::PageStarted(const TSharedRef<class IWebBrowserWindow>& InWindow)
 {
-	WindowPtr = InWindow;
+	SetWindow(InWindow);
 
-	// Only inject if page state matches requested one, either when page starts loading (bPageStarted = true) or when page completely loaded (bPageStarted = false)
-	if (bInjectJSOnPageStarted == bPageStarted)
+	if (bInjectJSOnPageStarted )
 	{
-		// Expunge temporary objects.
-		for (TMap<TObjectPtr<UObject>, ObjectBinding>::TIterator It(BoundObjects); It; ++It)
-		{
-			if (!It->Value.bIsPermanent)
-			{
-				It.RemoveCurrent();
-			}
-		}
+		InjectJavascript(InWindow);
+	}
+}
+
+void FMobileJSScripting::PageLoaded(const TSharedRef<class IWebBrowserWindow>& InWindow)
+{
+	SetWindow(InWindow);
 
-		InitializeScript(InWindow);
-		FString Script;
-		for (auto& Item : PermanentUObjectsByName)
+	if (!bInjectJSOnPageStarted)
+	{
+		InjectJavascript(InWindow);
+	}
+}
+
+void FMobileJSScripting::InjectJavascript(const TSharedRef<class IWebBrowserWindow>& Window)
+{
+	// Expunge temporary objects.
+	for (TMap<TObjectPtr<UObject>, ObjectBinding>::TIterator It(BoundObjects); It; ++It)
+	{
+		if (!It->Value.bIsPermanent)
 		{
-			Script.Append(*FString::Printf(TEXT("window.ue['%s'] = %s;"), *Item.Key.ReplaceCharWithEscapedChar(), *ConvertObject(Item.Value)));
+			It.RemoveCurrent();
 		}
-		InWindow->ExecuteJavascript(Script);
 	}
+
+	InitializeScript(Window);
+	FString Script;
+	for (auto& Item : PermanentUObjectsByName)
+	{
+		Script.Append(*FString::Printf(TEXT("window.ue['%s'] = %s;"), *Item.Key.ReplaceCharWithEscapedChar(), *ConvertObject(Item.Value)));
+	}
+	Window->ExecuteJavascript(Script);
 }
 
 FMobileJSScripting::FMobileJSScripting(bool bJSBindingToLoweringEnabled, bool bInjectJSOnPageStarted)
@@ -645,7 +659,7 @@ FMobileJSScripting::FMobileJSScripting(bool bJSBindingToLoweringEnabled, bool bI
 {
 }
 
-void FMobileJSScripting::SetWindow(TSharedRef<class IWebBrowserWindow> InWindow)
+void FMobileJSScripting::SetWindow(const TSharedRef<class IWebBrowserWindow>& InWindow)
 {
 	WindowPtr = InWindow;
 }
diff --git a/Engine/Source/Runtime/WebBrowser/Private/MobileJS/MobileJSScripting.h b/Engine/Source/Runtime/WebBrowser/Private/MobileJS/MobileJSScripting.h
--- a/Engine/Source/Runtime/WebBrowser/Private/MobileJS/MobileJSScripting.h
+++ b/Engine/Source/Runtime/WebBrowser/Private/MobileJS/MobileJSScripting.h
@@ -26,8 +26,8 @@ public:
 
 	virtual void BindUObject(const FString& Name, UObject* Object, bool bIsPermanent = true) override;
 	virtual void UnbindUObject(const FString& Name, UObject* Object = nullptr, bool bIsPermanent = true) override;
-	void BindUObject(TSharedRef<class IWebBrowserWindow> InWindow, const FString& Name, UObject* Object, bool bIsPermanent = true);
-	void UnbindUObject(TSharedRef<class IWebBrowserWindow> InWindow, const FString& Name, UObject* Object = nullptr, bool bIsPermanent = true);
+	void BindUObject(const TSharedRef<class IWebBrowserWindow>& InWindow, const FString& Name, UObject* Object, bool bIsPermanent = true);
+	void UnbindUObject(const TSharedRef<class IWebBrowserWindow>& InWindow, const FString& Name, UObject* Object = nullptr, bool bIsPermanent = true);
 
 	/**
 	 * Called when a message was received from the browser process.
@@ -43,12 +43,13 @@ public:
 
 	virtual void InvokeJSFunction(FGuid FunctionId, int32 ArgCount, FWebJSParam Arguments[], bool bIsError=false) override;
 	virtual void InvokeJSErrorResult(FGuid FunctionId, const FString& Error) override;
-	void TryInjectJavascript(TSharedRef<class IWebBrowserWindow> InWindow, bool bPageStarted);
-
-	void SetWindow(TSharedRef<class IWebBrowserWindow> InWindow);
+	void PageStarted(const TSharedRef<class IWebBrowserWindow>& InWindow);  // Called when page load started
+	void PageLoaded(const TSharedRef<class IWebBrowserWindow>& InWindow);  // Called on page completely loaded
+	void SetWindow(const TSharedRef<class IWebBrowserWindow>& InWindow);
 
 private:
-	void InitializeScript(TSharedRef<class IWebBrowserWindow> InWindow);
+	void InitializeScript(const TSharedRef<class IWebBrowserWindow>& InWindow);
+	void InjectJavascript(const TSharedRef<class IWebBrowserWindow>& InWindow);
 	void InvokeJSFunctionRaw(FGuid FunctionId, const FString& JSValue, bool bIsError=false);
 	bool IsValid()
 	{
