WebBrowser: Fix stretched out contents after a resize

This is a CEF regression that was introduced around v127, upstream issue:
[REDACTED_PROFILE_URL]/cef/issues/3826

We work around it first by triggering a repaint after each call to `WasResized()`, this restores the normal behavior.

There's a second issue however: the frame produced by that repaint sometimes has an old size, so we end up with slightly stretched contents or with black bars on the side. This is especially noticeable at the end of a resize, as it remains until another repaint is triggered (e.g. by a user action).

To fix that issue we check the size of the received frame against the size of our viewport. If they don't match we skip the frame and force another repaint.

#tests Added a WebBrowser in a Slate UI then resized the window/viewport, verified that the web contents adjust to the new sizes
#rnx

diff --git a/Engine/Source/Runtime/WebBrowser/Private/CEF/CEFWebBrowserWindow.cpp b/Engine/Source/Runtime/WebBrowser/Private/CEF/CEFWebBrowserWindow.cpp
--- a/Engine/Source/Runtime/WebBrowser/Private/CEF/CEFWebBrowserWindow.cpp
+++ b/Engine/Source/Runtime/WebBrowser/Private/CEF/CEFWebBrowserWindow.cpp
@@ -1975,6 +1975,11 @@ void FCEFWebBrowserWindow::HandleRenderingError()
 
 void FCEFWebBrowserWindow::OnPaint(CefRenderHandler::PaintElementType Type, const CefRenderHandler::RectList& DirtyRects, const void* Buffer, int Width, int Height)
 {
+	if (ViewportSize == FIntPoint::ZeroValue)
+	{
+		return;
+	}
+
 	bool bNeedsRedraw = false;
 	if (bUsingAcceleratedPaint)
 	{
@@ -2013,21 +2018,36 @@ void FCEFWebBrowserWindow::OnPaint(CefRenderHandler::PaintElementType Type, cons
 		}
 		else
 		{
-			UpdatableTextures[Type]->UpdateTextureThreadSafeRaw(Width, Height, Buffer, Dirty);
-			HandleRenderingError();
+			// Ignore the frame and force a CEF repaint if the texture size doesn't match the current viewport size.
+			// This is another side-effect of [REDACTED_PROFILE_URL]/cef/issues/3826:
+			// When resizing the view, the last frame we receive occasionally still has a previous viewport size,
+			// which results in contents looking stretched or with black borders.
+#if CEF_VERSION_MAJOR >= 128
+			if (Type == PET_VIEW && ViewportSize != FIntPoint(Width / ViewportDPIScaleFactor, Height / ViewportDPIScaleFactor))
+			{
+				// Note that calling Invalidate() doesn't work here, and there's no other CEF API to trigger a repaint.
+				InternalCefBrowser->GetHost()->WasHidden(true);
+				InternalCefBrowser->GetHost()->WasHidden(false);
+			}
+			else
+#endif
+			{
+				UpdatableTextures[Type]->UpdateTextureThreadSafeRaw(Width, Height, Buffer, Dirty);
+				HandleRenderingError();
+				bNeedsRedraw = true;
+			}
 
-		    if (Type == PET_POPUP && bShowPopupRequested)
-		    {
-			    bShowPopupRequested = false;
-			    bPopupHasFocus = true;
+			if (Type == PET_POPUP && bShowPopupRequested)
+			{
+				bShowPopupRequested = false;
+				bPopupHasFocus = true;
 
 				const float DPIScale = FPlatformApplicationMisc::GetDPIScaleFactorAtPoint(PopupPosition.X, PopupPosition.Y);
-			    FIntPoint PopupSize = FIntPoint(Width / DPIScale, Height / DPIScale);
+				FIntPoint PopupSize = FIntPoint(Width / DPIScale, Height / DPIScale);
 
-			    FIntRect PopupRect = FIntRect(PopupPosition, PopupPosition + PopupSize);
-			    OnShowPopup().Broadcast(PopupRect);
-		    }
-			bNeedsRedraw = true;
+				FIntRect PopupRect = FIntRect(PopupPosition, PopupPosition + PopupSize);
+				OnShowPopup().Broadcast(PopupRect);
+			}
 		}
 	}
 
@@ -2687,6 +2707,17 @@ void FCEFWebBrowserWindow::CheckTickActivity()
 	{
 		bNeedsResize = false;
 		InternalCefBrowser->GetHost()->WasResized();
+
+		// Recent CEF versions fail to trigger a repaint after WasResized() is called
+		// so we must force one on the view (and the popup if visible):
+		// (upstream issue: [REDACTED_PROFILE_URL]/cef/issues/3826)
+#if CEF_VERSION_MAJOR >= 128
+		InternalCefBrowser->GetHost()->Invalidate(PET_VIEW);
+		if (bPopupHasFocus)
+		{
+			InternalCefBrowser->GetHost()->Invalidate(PET_POPUP);
+		}
+#endif
 	}
 
 	bTickedLastFrame = false;
