WebBrowser: prevent showing a Chromium browser window in a rare case of concurrent process runs

This could happen in theory if two UE processes initialize CEF at the exact same time, so the second process doesn't see the 'lockfile' created by the first process in the cache dir.
CEF has a built-in app singleton mechanism in this case which reports in the first process that a new app instance was started. By default a regular Chromium browser window is created, which could be confusing to users or even raise suspicions.

Implement the `OnAlreadyRunningAppRelaunch` handler to disable the default behavior and do nothing.

#tests forced a second process to start with the same cache, verified that no Chrome browser window appears
#rnx

diff --git a/Engine/Source/Runtime/WebBrowser/Private/CEF/CEFBrowserApp.cpp b/Engine/Source/Runtime/WebBrowser/Private/CEF/CEFBrowserApp.cpp
--- a/Engine/Source/Runtime/WebBrowser/Private/CEF/CEFBrowserApp.cpp
+++ b/Engine/Source/Runtime/WebBrowser/Private/CEF/CEFBrowserApp.cpp
@@ -123,4 +123,12 @@ bool FCEFBrowserApp::TickMessagePump(float DeltaTime, bool bForce)
 	return false;
 }
 
+#if CEF_VERSION_MAJOR >= 128
+bool FCEFBrowserApp::OnAlreadyRunningAppRelaunch(CefRefPtr<CefCommandLine> CommandLine, const CefString& CurrentDirectory)
+{
+	// nothing to do, except report the relaunch as handled to prevent the default handler from opening a new browser window
+	return true;
+}
+#endif
+
 #endif
diff --git a/Engine/Source/Runtime/WebBrowser/Private/CEF/CEFBrowserApp.h b/Engine/Source/Runtime/WebBrowser/Private/CEF/CEFBrowserApp.h
--- a/Engine/Source/Runtime/WebBrowser/Private/CEF/CEFBrowserApp.h
+++ b/Engine/Source/Runtime/WebBrowser/Private/CEF/CEFBrowserApp.h
@@ -37,6 +37,7 @@ private:
 	virtual void OnScheduleMessagePumpWork(int64 delay_ms) override;
 #else
 	virtual void OnScheduleMessagePumpWork(int64_t delay_ms) override;
+	virtual bool OnAlreadyRunningAppRelaunch(CefRefPtr<CefCommandLine> CommandLine, const CefString& CurrentDirectory) override;
 #endif
 
 	// Include the default reference counting implementation.
