WebBrowser: work around crash when multiple processes instantiate WebBrowserSingleton with the same cache path

#tests Started two editors in parallel with the same project, verified creation of a separate cef cache folder for the second instance
#rnx

diff --git a/Engine/Source/Runtime/WebBrowser/Private/WebBrowserSingleton.cpp b/Engine/Source/Runtime/WebBrowser/Private/WebBrowserSingleton.cpp
--- a/Engine/Source/Runtime/WebBrowser/Private/WebBrowserSingleton.cpp
+++ b/Engine/Source/Runtime/WebBrowser/Private/WebBrowserSingleton.cpp
@@ -914,7 +914,29 @@ FString FWebBrowserSingleton::GenerateWebCacheFolderName(const FString& InputPat
 
 	// append the version of this CEF build to our requested cache folder path
 	// this means each new CEF build gets its own cache folder, making downgrading safe
-	return InputPath + "_" + MAKE_STRING(CHROME_VERSION_BUILD);
+	FString VersionedCachePath = InputPath + "_" + MAKE_STRING(CHROME_VERSION_BUILD);
+
+#if CEF_VERSION_MAJOR >= 128
+	// CEF v128 prevents concurrent browser cache access between multiple processes.
+	// It locks the cache folder by creating a "lockfile" file in it during initialization,
+	// so that other CEF initializations will fail as long as the cache folder is locked.
+	// The lockfile is then deleted when the locking CEF process ends (even in case of a crash)
+	for (uint32 CacheNumber = 0; CacheNumber < 32; ++CacheNumber) {
+		FString ConcurrentCachePath = VersionedCachePath;
+		// don't append the first _0 so we can reuse caches created as VersionedCachePath
+		if (CacheNumber)
+		{
+			ConcurrentCachePath.Append(FString::Format(TEXT("_{0}"), { CacheNumber }));
+		}
+		if (!FPaths::FileExists(FPaths::Combine(ConcurrentCachePath, TEXT("lockfile"))))
+		{
+			return ConcurrentCachePath;
+		}
+	}
+	// just fall back to the base path and let CEF initialization fail if we've reached the limit
+#endif
+
+	return VersionedCachePath;
 }
 #endif
 
