Fix some memory leaks in the iOS web browser.

diff --git a/Engine/Source/Runtime/WebBrowser/Private/Android/AndroidWebBrowserWidget.cpp b/Engine/Source/Runtime/WebBrowser/Private/Android/AndroidWebBrowserWidget.cpp
--- a/Engine/Source/Runtime/WebBrowser/Private/Android/AndroidWebBrowserWidget.cpp
+++ b/Engine/Source/Runtime/WebBrowser/Private/Android/AndroidWebBrowserWidget.cpp
@@ -88,6 +88,11 @@ SAndroidWebBrowserWidget::~SAndroidWebBrowserWidget()
 	
 	WebBrowserTextureSamplesQueue->RequestFlush();
 
+	if (WebBrowserBrush.IsValid())
+	{
+		WebBrowserBrush->SetResourceObject(nullptr);	
+	}
+	
 	if (WebBrowserMaterial != nullptr)
 	{
 		WebBrowserMaterial->RemoveFromRoot();
diff --git a/Engine/Source/Runtime/WebBrowser/Private/IOS/IOSPlatformWebBrowser.cpp b/Engine/Source/Runtime/WebBrowser/Private/IOS/IOSPlatformWebBrowser.cpp
--- a/Engine/Source/Runtime/WebBrowser/Private/IOS/IOSPlatformWebBrowser.cpp
+++ b/Engine/Source/Runtime/WebBrowser/Private/IOS/IOSPlatformWebBrowser.cpp
@@ -293,6 +293,30 @@ class SIOSWebBrowserWidget : public SLeafWidget
 			WebViewWrapper = nil;
 		}
 		WebBrowserWindowPtr.Reset();
+
+#if !PLATFORM_TVOS
+		delete TextureSamplePool;
+		TextureSamplePool = nullptr;
+		
+		WebBrowserTextureSamplesQueue->RequestFlush();
+
+		if (WebBrowserBrush.IsValid())
+		{
+			WebBrowserBrush->SetResourceObject(nullptr);	
+		}
+
+		if (WebBrowserMaterial != nullptr)
+		{
+			WebBrowserMaterial->RemoveFromRoot();
+			WebBrowserMaterial = nullptr;
+		}
+
+		if (WebBrowserTexture != nullptr)
+		{
+			WebBrowserTexture->RemoveFromRoot();
+			WebBrowserTexture = nullptr;
+		}
+#endif
 	}
 
 	void ShowFloatingCloseButton(bool bShow, bool bDraggable)
