Objective-C wrapper object leaking on shutdow of iOS web browser. Ensure it is released.

diff --git a/Engine/Source/Runtime/WebBrowser/Private/IOS/IOSPlatformWebBrowser.cpp b/Engine/Source/Runtime/WebBrowser/Private/IOS/IOSPlatformWebBrowser.cpp
--- a/Engine/Source/Runtime/WebBrowser/Private/IOS/IOSPlatformWebBrowser.cpp
+++ b/Engine/Source/Runtime/WebBrowser/Private/IOS/IOSPlatformWebBrowser.cpp
@@ -290,6 +290,7 @@ class SIOSWebBrowserWidget : public SLeafWidget
 		if (WebViewWrapper != nil)
 		{
 			[WebViewWrapper close];
+			[WebViewWrapper release];
 			WebViewWrapper = nil;
 		}
 		WebBrowserWindowPtr.Reset();
