Adapted IOSPlatformWebBrowser to ApplePlatformWebBrowser so that macOS can use WebKit instead of CEF3. Using WebKit when bCompileCEF3 build flag is False.

diff --git a/Engine/Config/Mac/BaseMacEngine.ini b/Engine/Config/Mac/BaseMacEngine.ini
--- a/Engine/Config/Mac/BaseMacEngine.ini
+++ b/Engine/Config/Mac/BaseMacEngine.ini
@@ -61,3 +61,6 @@ CycleDetailsViewBind=(Key=Add,bShift=False,bCtrl=False,bAlt=False,bCmd=False)
 DebugCameraBind=(Key=Tab,bShift=False,bCtrl=False,bAlt=False,bCmd=False)
 OnScreenDebugMessagesBind=(Key=Tab,bShift=True,bCtrl=False,bAlt=False,bCmd=False)
 GameHUDBind=(Key=Tilde,bShift=True,bCtrl=False,bAlt=False,bCmd=False)
+
+[/Script/BuildSettings.BuildSettings]
+bCompileCEF3=True
diff --git a/Engine/Source/Runtime/Launch/Private/LaunchEngineLoop.cpp b/Engine/Source/Runtime/Launch/Private/LaunchEngineLoop.cpp
--- a/Engine/Source/Runtime/Launch/Private/LaunchEngineLoop.cpp
+++ b/Engine/Source/Runtime/Launch/Private/LaunchEngineLoop.cpp
@@ -248,6 +248,10 @@
 	#include "Tests/RHIUnitTests.h"
 #endif
 
+#if PLATFORM_MAC
+	#include "Apple/AppleAsyncTask.h"
+#endif
+
 class FSlateRenderer;
 class SViewport;
 class IPlatformFile;
@@ -6031,6 +6035,11 @@ void FEngineLoop::Tick()
 		FEmbeddedCommunication::AllowSleep(TEXT("FirstTicks"));
 	}
 #endif
+	
+#if PLATFORM_MAC
+	// drain the apple async task queue from the game thread
+	[FAppleAsyncTask ProcessAsyncTasks];
+#endif
 }
 
 
diff --git a/Engine/Source/Runtime/WebBrowser/Private/Apple/AppleCookieManager.cpp b/Engine/Source/Runtime/WebBrowser/Private/Apple/AppleCookieManager.cpp
new file mode 100644
--- /dev/null
+++ b/Engine/Source/Runtime/WebBrowser/Private/Apple/AppleCookieManager.cpp
@@ -0,0 +1,74 @@
+// Copyright Epic Games, Inc. All Rights Reserved.
+
+#include "AppleCookieManager.h"
+
+#if PLATFORM_IOS || PLATFORM_MAC
+#import <Foundation/Foundation.h>
+#include "Apple/AppleAsyncTask.h"
+
+#if !PLATFORM_TVOS
+#import "WebKit/WebKit.h"
+#endif
+
+NSUUID *FAppleCookieManager::CookieManagerIdentifier = [[NSUUID alloc] initWithUUIDString: @"0986df82-fc1b-4dd1-864d-e149a7e79119"];
+
+FAppleCookieManager::FAppleCookieManager()
+{
+}
+
+FAppleCookieManager::~FAppleCookieManager()
+{
+}
+
+void FAppleCookieManager::SetCookie(const FString& URL, const FCookie& Cookie, TFunction<void(bool)> Completed)
+{
+	// not implemented
+	if (Completed)
+	{
+		Completed(false);
+	}
+}
+
+void FAppleCookieManager::DeleteCookies(const FString& URL, const FString& CookieName, TFunction<void(int)> Completed)
+{
+#if !PLATFORM_TVOS
+	FString CapturedURL = URL;
+	UE_LOG(LogTemp, Warning, TEXT("Deleting cookies for %s"), *URL);
+	
+	dispatch_async(dispatch_get_main_queue(), ^
+	{
+		// Delete matching cookies
+		WKHTTPCookieStore* Storage = [WKWebsiteDataStore dataStoreForIdentifier: CookieManagerIdentifier].httpCookieStore;
+		bool CapturedURLEmpty = CapturedURL.IsEmpty();
+		
+		[Storage getAllCookies: ^(NSArray<NSHTTPCookie*>* Cookies) 
+		{
+			for (NSHTTPCookie* Cookie in Cookies)
+			{
+				FString Domain([Cookie domain]);
+				FString Path([Cookie path]);
+				FString CookieUrl = Domain + Path;
+				FString URLTest = CapturedURL;
+				
+				if (CookieUrl.Contains(CapturedURL) || CapturedURLEmpty)
+				{
+					[Storage deleteCookie:Cookie completionHandler:nil];
+				}
+			}
+			
+			// Notify on the game thread
+			[FAppleAsyncTask CreateTaskWithBlock:^bool(void)
+			 {
+				 if (Completed)
+				 {
+					 Completed(true);
+				 }
+				 return true;
+			 }];
+		}];
+		
+	});
+#endif
+}
+
+#endif
diff --git a/Engine/Source/Runtime/WebBrowser/Private/Apple/AppleCookieManager.h b/Engine/Source/Runtime/WebBrowser/Private/Apple/AppleCookieManager.h
new file mode 100644
--- /dev/null
+++ b/Engine/Source/Runtime/WebBrowser/Private/Apple/AppleCookieManager.h
@@ -0,0 +1,31 @@
+// Copyright Epic Games, Inc. All Rights Reserved.
+
+#pragma once
+
+#include "CoreMinimal.h"
+
+#if PLATFORM_IOS || PLATFORM_MAC
+#include "IWebBrowserCookieManager.h"
+
+/**
+ * Implementation of interface for dealing with a Web Browser cookies for iOS.
+ */
+class FAppleCookieManager
+	: public IWebBrowserCookieManager
+	, public TSharedFromThis<FAppleCookieManager>
+{
+public:
+	// This is a hack for now until I figure out how to persist cookie managers' UUIDs
+	static NSUUID *CookieManagerIdentifier;
+	
+	// IWebBrowserCookieManager interface
+
+	virtual void SetCookie(const FString& URL, const FCookie& Cookie, TFunction<void(bool)> Completed = nullptr) override;
+	virtual void DeleteCookies(const FString& URL = TEXT(""), const FString& CookieName = TEXT(""), TFunction<void(int)> Completed = nullptr) override;
+
+	// FAppleCookieManager
+
+	FAppleCookieManager();
+	virtual ~FAppleCookieManager();
+};
+#endif
diff --git a/Engine/Source/Runtime/WebBrowser/Private/Apple/ApplePlatformWebBrowser.cpp b/Engine/Source/Runtime/WebBrowser/Private/Apple/ApplePlatformWebBrowser.cpp
new file mode 100644
--- /dev/null
+++ b/Engine/Source/Runtime/WebBrowser/Private/Apple/ApplePlatformWebBrowser.cpp
@@ -0,0 +1,1489 @@
+// Copyright Epic Games, Inc. All Rights Reserved.
+
+#include "ApplePlatformWebBrowser.h"
+
+#if !WITH_CEF3 && (PLATFORM_IOS || PLATFORM_MAC)
+
+#if PLATFORM_IOS
+#include "IOS/IOSView.h"
+#include "IOS/IOSAppDelegate.h"
+#else
+#include "Mac/CocoaWindow.h"
+#include "Mac/MacWindow.h"
+#include "Mac/MacApplication.h"
+#endif
+
+#include "Widgets/SLeafWidget.h"
+#include "PlatformHttp.h"
+#include "HAL/PlatformProcess.h"
+
+#import <MetalKit/MetalKit.h>
+#include "ExternalTexture.h"
+#include "WebBrowserModule.h"
+#include "WebViewCloseButton.h"
+#include "IWebBrowserSingleton.h"
+
+#include "Apple/AppleAsyncTask.h"
+#include "AppleCookieManager.h"
+
+DEFINE_LOG_CATEGORY(LogAppleWebBrowser);
+
+#if !UE_BUILD_SHIPPING
+TAutoConsoleVariable<bool> CVarWebViewEnableInspectable(
+	TEXT("webview.EnableInspectable"),
+	false,
+	TEXT("Allows webview to be debuggable through Safari Web Inspector."),
+	ECVF_Default
+);
+#endif
+
+class SAppleWebBrowserWidget : public SLeafWidget
+{
+	SLATE_BEGIN_ARGS(SAppleWebBrowserWidget)
+		: _InitialURL("about:blank")
+		, _UseTransparency(false)
+	{ }
+
+	SLATE_ARGUMENT(FString, InitialURL);
+	SLATE_ARGUMENT(bool, UseTransparency);
+	SLATE_ARGUMENT(TSharedPtr<FWebBrowserWindow>, WebBrowserWindow);
+	SLATE_ARGUMENT(FString, UserAgentApplication);
+
+	SLATE_END_ARGS()
+
+	SAppleWebBrowserWidget()
+		: WebViewWrapper(nil)
+	{}
+
+	void Construct(const FArguments& Args)
+	{
+		bool bSupportsMetalMRT = false;
+		GConfig->GetBool(TEXT("/Script/IOSRuntimeSettings.IOSRuntimeSettings"), TEXT("bSupportsMetalMRT"), bSupportsMetalMRT, GEngineIni);
+
+		bool bSupportsMetal = false;
+		GConfig->GetBool(TEXT("/Script/IOSRuntimeSettings.IOSRuntimeSettings"), TEXT("bSupportsMetal"), bSupportsMetal, GEngineIni);
+		// At this point we MUST be a Metal renderer.
+		check(bSupportsMetal);
+
+		WebBrowserWindowPtr = Args._WebBrowserWindow;
+		Is3DBrowser = false;
+
+		bool bEnableFloatingCloseButton = false;
+		GConfig->GetBool(TEXT("Browser"), TEXT("bEnableFloatingCloseButton"), bEnableFloatingCloseButton, GEngineIni);
+
+		TSharedPtr<SAppleWebBrowserWidget> SharedThis(this);
+		WebViewWrapper = [AppleWebViewWrapper alloc];
+		[WebViewWrapper create : SharedThis userAgentApplication : Args._UserAgentApplication.GetNSString() useTransparency : Args._UseTransparency supportsMetal : bSupportsMetal supportsMetalMRT : bSupportsMetalMRT enableFloatingCloseButton : bEnableFloatingCloseButton];
+
+#if !PLATFORM_TVOS
+		OnSafeFrameChangedEventHandle = FCoreDelegates::OnSafeFrameChangedEvent.AddLambda( [this, WeakThis = TWeakPtr<SAppleWebBrowserWidget>(SharedThis)]()
+										{
+											if (TSharedPtr<SAppleWebBrowserWidget> StrongThis = WeakThis.Pin())
+											{
+												if (WebViewWrapper != nil)
+												{
+													[WebViewWrapper didRotate];
+												}
+											}
+										});
+
+
+		TextureSamplePool = new FWebBrowserTextureSamplePool();
+		WebBrowserTextureSamplesQueue = MakeShared<FWebBrowserTextureSampleQueue, ESPMode::ThreadSafe>();
+		WebBrowserTexture = nullptr;
+		WebBrowserMaterial = nullptr;
+		WebBrowserBrush = nullptr;
+
+		// create external texture
+		WebBrowserTexture = NewObject<UWebBrowserTexture>((UObject*)GetTransientPackage(), NAME_None, RF_Transient | RF_Public);
+
+		if (WebBrowserTexture != nullptr)
+		{
+			WebBrowserTexture->UpdateResource();
+			WebBrowserTexture->AddToRoot();
+		}
+
+		// create wrapper material
+		IWebBrowserSingleton* WebBrowserSingleton = IWebBrowserModule::Get().GetSingleton();
+
+		UMaterialInterface* DefaultWBMaterial = Args._UseTransparency ? WebBrowserSingleton->GetDefaultTranslucentMaterial() : WebBrowserSingleton->GetDefaultMaterial();
+		if (WebBrowserSingleton && DefaultWBMaterial)
+		{
+			// create wrapper material
+			WebBrowserMaterial = UMaterialInstanceDynamic::Create(DefaultWBMaterial, nullptr);
+
+			if (WebBrowserMaterial)
+			{
+				WebBrowserMaterial->SetTextureParameterValue("SlateUI", WebBrowserTexture);
+				WebBrowserMaterial->AddToRoot();
+
+				// create Slate brush
+				WebBrowserBrush = MakeShareable(new FSlateBrush());
+				{
+					WebBrowserBrush->SetResourceObject(WebBrowserMaterial);
+				}
+			}
+		}
+#endif
+		LoadURL(Args._InitialURL);
+	}
+
+	void UpdateFrameWithViewport(FIntPoint WindowSize, FIntPoint WindowPos) {
+		CGRect NewFrame;
+#if PLATFORM_MAC
+		TSharedPtr<FGenericWindow> EnclosedWindow = WebBrowserWindowPtr.Pin()->GetParentWindow().Get()->GetNativeWindow();
+		{
+			NSWindow* ParentWindow = (NSWindow*)EnclosedWindow->GetOSWindowHandle();
+			
+			const FVector2D CocoaPosition = FMacApplication::ConvertSlatePositionToCocoa(WindowPos.X, WindowPos.Y);
+			
+			NSRect ParentFrame = [ParentWindow frame];
+			NSRect Rect = NSMakeRect(CocoaPosition.X - ParentFrame.origin.x, (CocoaPosition.Y - ParentFrame.origin.y) - WindowSize.Y, 
+									 FMath::Max(WindowSize.X, 1), FMath::Max(WindowSize.Y, 1));
+			Rect = [ParentWindow frameRectForContentRect : Rect];
+			
+			// CGRect and NSRect have the same memory layout
+			NewFrame = NSRectToCGRect(Rect);
+		}
+#else
+		{
+			UIView* View = [IOSAppDelegate GetDelegate].IOSView;
+			CGFloat ContentScaleFactor = View.contentScaleFactor;
+				
+			FVector2D PosVector2D = WindowPos * ContentScaleFactor;
+			
+			WindowPos /= ContentScaleFactor;
+			WindowSize /= ContentScaleFactor;
+			
+			NewFrame.size.width = WindowSize.X;
+			NewFrame.size.height = WindowSize.Y;
+			
+			NewFrame.origin.x = WindowPos.X;
+			NewFrame.origin.y = WindowPos.Y;
+		}
+#endif
+
+		if(NewFrame.size.width != WindowFrame.size.width || NewFrame.size.height != WindowFrame.size.height) {
+			WindowFrame = NewFrame;
+		}
+		
+#if PLATFORM_IOS
+		[WebViewWrapper updateframe : WindowFrame window : nil];
+#else			
+		[WebViewWrapper updateframe : WindowFrame window : EnclosedWindow.Get()];
+#endif
+	}
+	
+	void Tick(const FGeometry& AllottedGeometry, const double InCurrentTime, const float InDeltaTime)
+	{
+		if (WebViewWrapper != nil)
+		{
+			if (WebBrowserWindowPtr.IsValid())
+			{
+				WebBrowserWindowPtr.Pin()->SetTickLastFrame();
+				if (WebBrowserWindowPtr.Pin()->GetParentWindow().IsValid())
+				{
+					bool ShouldSet3DBrowser = WebBrowserWindowPtr.Pin()->GetParentWindow().Get()->IsVirtualWindow();
+					if (Is3DBrowser != ShouldSet3DBrowser)
+					{
+						Is3DBrowser = ShouldSet3DBrowser;
+						[WebViewWrapper set3D : Is3DBrowser];
+					}
+				}
+			}
+			
+			FVector2D Position = AllottedGeometry.GetAccumulatedRenderTransform().GetTranslation();
+			FVector2D Size = TransformVector(AllottedGeometry.GetAccumulatedRenderTransform(), AllottedGeometry.GetLocalSize());
+			FSlateRect SlateRect = AllottedGeometry.GetRenderBoundingRect();
+			
+			UpdateFrameWithViewport(FIntPoint(static_cast<int32_t>(Size.X), static_cast<int32_t>(Size.Y)), 
+									FIntPoint(static_cast<int32_t>(Position.X), static_cast<int32_t>(Position.Y)));
+			
+#if !PLATFORM_TVOS
+#if !UE_BUILD_SHIPPING
+			if (WebViewWrapper.WebView != nil)
+			{
+				WebViewWrapper.WebView.inspectable = CVarWebViewEnableInspectable.GetValueOnAnyThread()? YES : NO;
+			}
+#endif
+			if (Is3DBrowser)
+			{
+				if (WebBrowserTexture)
+				{					
+					TSharedPtr<FWebBrowserTextureSample, ESPMode::ThreadSafe> WebBrowserTextureSample;
+					WebBrowserTextureSamplesQueue->Peek(WebBrowserTextureSample);
+
+					WebBrowserTexture->TickResource(WebBrowserTextureSample);
+				}
+
+				if (WebBrowserTexture != nullptr)
+				{
+					struct FWriteWebBrowserParams
+					{
+						AppleWebViewWrapper* NativeWebBrowserPtr;
+						FGuid PlayerGuid;
+						FIntPoint Size;
+					};
+
+					FIntPoint viewportSize = WebBrowserWindowPtr.Pin()->GetViewportSize();
+
+					FWriteWebBrowserParams Params = { WebViewWrapper, WebBrowserTexture->GetExternalTextureGuid(), viewportSize };
+
+					ENQUEUE_RENDER_COMMAND(WriteWebBrowser)(
+						[Params](FRHICommandListImmediate& RHICmdList)
+						{
+							AppleWebViewWrapper* NativeWebBrowser = Params.NativeWebBrowserPtr;
+
+							if (NativeWebBrowser == nil)
+							{
+								return;
+							}
+
+							FTextureRHIRef VideoTexture = [NativeWebBrowser GetVideoTexture];
+							if (VideoTexture == nullptr)
+							{
+								const FRHITextureCreateDesc Desc =
+									FRHITextureCreateDesc::Create2D(TEXT("SAppleWebBrowserWidget_VideoTexture"), Params.Size, PF_R8G8B8A8)
+									.SetFlags(ETextureCreateFlags::External);
+
+								VideoTexture = RHICreateTexture(Desc);
+								[NativeWebBrowser SetVideoTexture : VideoTexture];
+
+								if (VideoTexture == nullptr)
+								{
+									UE_LOG(LogAppleWebBrowser, Warning, TEXT("RHICreateTexture failed!"));
+									return;
+								}
+
+								[NativeWebBrowser SetVideoTextureValid : false];
+
+							}
+
+							if ([NativeWebBrowser UpdateVideoFrame : VideoTexture->GetNativeResource()])
+							{
+								// if region changed, need to reregister UV scale/offset
+								//UE_LOG(LogAppleWebBrowser, Log, TEXT("UpdateVideoFrame RT: %s"), *Params.PlayerGuid.ToString());
+							}
+
+							if (![NativeWebBrowser IsVideoTextureValid])
+							{
+								FSamplerStateInitializerRHI SamplerStateInitializer(SF_Bilinear, AM_Clamp, AM_Clamp, AM_Clamp);
+								FSamplerStateRHIRef SamplerStateRHI = RHICreateSamplerState(SamplerStateInitializer);
+								FExternalTextureRegistry::Get().RegisterExternalTexture(Params.PlayerGuid, VideoTexture, SamplerStateRHI);
+
+								[NativeWebBrowser SetVideoTextureValid : true];
+							}
+						});
+				}
+			}
+#endif
+		}
+
+	}
+
+	int32 OnPaint(const FPaintArgs& Args, const FGeometry& AllottedGeometry, const FSlateRect& MyCullingRect, FSlateWindowElementList& OutDrawElements, int32 LayerId, const FWidgetStyle& InWidgetStyle, bool bParentEnabled) const
+	{
+#if !PLATFORM_TVOS
+		bool bIsVisible = !WebBrowserWindowPtr.IsValid() || WebBrowserWindowPtr.Pin()->IsVisible();
+		
+		if (bIsVisible && Is3DBrowser && WebBrowserBrush.IsValid())
+		{
+			FSlateDrawElement::MakeBox(OutDrawElements, LayerId, AllottedGeometry.ToPaintGeometry(), WebBrowserBrush.Get(), ESlateDrawEffect::None);
+		}
+#endif
+		return LayerId;
+	}
+
+	virtual FVector2D ComputeDesiredSize(float) const override
+	{
+		return FVector2D(640, 480);
+	}
+
+	void LoadURL(const FString& InNewURL)
+	{
+		if (WebViewWrapper != nil)
+		{
+			[WebViewWrapper loadurl : [NSURL URLWithString : [NSString stringWithUTF8String : TCHAR_TO_UTF8(*InNewURL)]]];
+		}
+	}
+
+	void LoadString(const FString& InContents, const FString& InDummyURL)
+	{
+		if (WebViewWrapper != nil)
+		{
+			[WebViewWrapper loadstring : [NSString stringWithUTF8String : TCHAR_TO_UTF8(*InContents)] dummyurl : [NSURL URLWithString : [NSString stringWithUTF8String : TCHAR_TO_UTF8(*InDummyURL)]]];
+		}
+	}
+	
+	void StopLoad()
+	{
+		if (WebViewWrapper != nil)
+		{
+			[WebViewWrapper stopLoading];
+		}
+	}
+
+	void Reload()
+	{
+		if (WebViewWrapper != nil)
+		{
+			[WebViewWrapper reload];
+		}
+	}
+
+	void Close()
+	{
+#if !PLATFORM_TVOS
+		FCoreDelegates::OnSafeFrameChangedEvent.Remove(OnSafeFrameChangedEventHandle);
+#endif
+		if (WebViewWrapper != nil)
+		{
+			[WebViewWrapper close];
+			[WebViewWrapper release];
+			WebViewWrapper = nil;
+		}
+		WebBrowserWindowPtr.Reset();
+		
+#if !PLATFORM_TVOS
+		delete TextureSamplePool;
+		TextureSamplePool = nullptr;
+		
+		WebBrowserTextureSamplesQueue->RequestFlush();
+		if (WebBrowserBrush.IsValid())
+		{
+			WebBrowserBrush->SetResourceObject(nullptr);	
+		}
+		if (WebBrowserMaterial != nullptr)
+		{
+			WebBrowserMaterial->RemoveFromRoot();
+			WebBrowserMaterial = nullptr;
+		}
+		if (WebBrowserTexture != nullptr)
+		{
+			WebBrowserTexture->RemoveFromRoot();
+			WebBrowserTexture = nullptr;
+		}
+#endif
+	}
+
+	void ShowFloatingCloseButton(bool bShow, bool bDraggable)
+	{
+		if (WebViewWrapper != nil)
+		{
+			[WebViewWrapper showFloatingCloseButton: bShow setDraggable: bDraggable];
+		}
+	}
+	
+	void GoBack()
+	{
+		if (WebViewWrapper != nil)
+		{
+			[WebViewWrapper goBack];
+		}
+	}
+
+	void GoForward()
+	{
+		if (WebViewWrapper != nil)
+		{
+			[WebViewWrapper goForward];
+		}
+	}
+
+
+	bool CanGoBack()
+	{
+		if (WebViewWrapper != nil)
+		{
+			return [WebViewWrapper canGoBack];
+		}
+		return false;
+	}
+
+	bool CanGoForward()
+	{
+		if (WebViewWrapper != nil)
+		{
+			return [WebViewWrapper canGoForward];
+		}
+		return false;
+	}
+
+	void SetWebBrowserVisibility(bool InIsVisible)
+	{
+		if (WebViewWrapper != nil)
+		{
+			UE_LOG(LogAppleWebBrowser, Warning, TEXT("SetWebBrowserVisibility %d!"), InIsVisible);
+
+			[WebViewWrapper setVisibility : InIsVisible];
+		}
+	}
+
+	bool HandleOnBeforePopup(const FString& UrlStr, const FString& FrameName)
+	{
+		TSharedPtr<FWebBrowserWindow> BrowserWindow = WebBrowserWindowPtr.Pin();
+		if (BrowserWindow.IsValid() && BrowserWindow->OnBeforePopup().IsBound())
+		{
+			return BrowserWindow->OnBeforePopup().Execute(UrlStr, FrameName);
+		}
+		return false;
+	}
+
+	bool HandleShouldOverrideUrlLoading(const FString& Url)
+	{
+		if (WebBrowserWindowPtr.IsValid())
+		{
+			// Capture vars needed for AsyncTask
+			NSString* UrlString = [NSString stringWithUTF8String : TCHAR_TO_UTF8(*Url)];
+			TWeakPtr<FWebBrowserWindow> AsyncWebBrowserWindowPtr = WebBrowserWindowPtr;
+
+			// Notify on the game thread
+			[FAppleAsyncTask CreateTaskWithBlock : ^ bool(void)
+			{
+				TSharedPtr<FWebBrowserWindow> BrowserWindow = AsyncWebBrowserWindowPtr.Pin();
+				if (BrowserWindow.IsValid())
+				{
+					if (BrowserWindow->OnBeforeBrowse().IsBound())
+					{
+						FWebNavigationRequest RequestDetails;
+						RequestDetails.bIsRedirect = false;
+						RequestDetails.bIsMainFrame = true; // shouldOverrideUrlLoading is only called on the main frame
+
+						BrowserWindow->OnBeforeBrowse().Execute(UrlString, RequestDetails);
+						BrowserWindow->SetTitle("");
+					}
+				}
+				return true;
+			}];
+		}
+		return true;
+	}
+
+	void HandleReceivedTitle(const FString& Title)
+	{
+		if (WebBrowserWindowPtr.IsValid())
+		{
+			TSharedPtr<FWebBrowserWindow> BrowserWindow = WebBrowserWindowPtr.Pin();
+			if (BrowserWindow.IsValid() && !BrowserWindow->GetTitle().Equals(Title))
+			{
+				BrowserWindow->SetTitle(Title);
+			}
+		}
+	}
+
+	void ProcessScriptMessage(const FString& InMessage)
+	{
+		if (WebBrowserWindowPtr.IsValid())
+		{
+			FString Message = InMessage;
+			TWeakPtr<FWebBrowserWindow> AsyncWebBrowserWindowPtr = WebBrowserWindowPtr;
+
+			[FAppleAsyncTask CreateTaskWithBlock : ^ bool(void)
+			{
+				TSharedPtr<FWebBrowserWindow> BrowserWindow = AsyncWebBrowserWindowPtr.Pin();
+				if (BrowserWindow.IsValid())
+				{
+					TArray<FString> Params;
+					Message.ParseIntoArray(Params, TEXT("/"), false);
+					if (Params.Num() > 0)
+					{
+						for (int I = 0; I < Params.Num(); I++)
+						{
+							Params[I] = FPlatformHttp::UrlDecode(Params[I]);
+						}
+
+						FString Command = Params[0];
+						Params.RemoveAt(0, 1);
+						BrowserWindow->OnJsMessageReceived(Command, Params, "");
+					}
+					else
+					{
+						GLog->Logf(ELogVerbosity::Error, TEXT("Invalid message from browser view: %s"), *Message);
+					}
+				}
+				return true;
+			}];
+		}
+	}
+
+	void HandlePageLoad(const FString& InCurrentUrl, bool bIsLoading)
+	{
+		if (WebBrowserWindowPtr.IsValid())
+		{
+			TSharedPtr<FWebBrowserWindow> BrowserWindow = WebBrowserWindowPtr.Pin();
+			if (BrowserWindow.IsValid())
+			{
+				BrowserWindow->NotifyDocumentLoadingStateChange(InCurrentUrl, bIsLoading);
+			}
+		}
+	}
+
+	void HandleReceivedError(int ErrorCode, const FString& InCurrentUrl)
+	{
+		if (WebBrowserWindowPtr.IsValid())
+		{
+			TSharedPtr<FWebBrowserWindow> BrowserWindow = WebBrowserWindowPtr.Pin();
+			if (BrowserWindow.IsValid())
+			{
+				BrowserWindow->NotifyDocumentError(InCurrentUrl, ErrorCode);
+			}
+		}
+	}
+
+	void ExecuteJavascript(const FString& Script)
+	{
+		if (WebViewWrapper != nil)
+		{
+			[WebViewWrapper executejavascript : [NSString stringWithUTF8String : TCHAR_TO_UTF8(*Script)]];
+		}
+	}
+
+	void FloatingCloseButtonPressed()
+	{
+		if (WebBrowserWindowPtr.IsValid())
+		{
+			TWeakPtr<FWebBrowserWindow> AsyncWebBrowserWindowPtr = WebBrowserWindowPtr;
+			
+#if !PLATFORM_MAC
+			[FAppleAsyncTask CreateTaskWithBlock : ^ bool(void)
+			{
+				if (TSharedPtr<FWebBrowserWindow> BrowserWindow = AsyncWebBrowserWindowPtr.Pin())
+				{
+					BrowserWindow->FloatingCloseButtonPressed();
+				}
+				return true;
+			}];
+#endif
+		}
+	}
+	
+	~SAppleWebBrowserWidget()
+	{
+		Close();
+	}
+
+protected:
+	mutable __strong AppleWebViewWrapper* WebViewWrapper;
+private:
+	TWeakPtr<FWebBrowserWindow> WebBrowserWindowPtr;
+
+	/** Enable 3D appearance */
+	bool Is3DBrowser;
+
+#if !PLATFORM_TVOS
+	/** The external texture to render the webbrowser output. */
+	UWebBrowserTexture* WebBrowserTexture;
+
+	/** The material for the external texture. */
+	UMaterialInstanceDynamic* WebBrowserMaterial;
+
+	/** The Slate brush that renders the material. */
+	TSharedPtr<FSlateBrush> WebBrowserBrush;
+
+	/** The sample queue. */
+	TSharedPtr<FWebBrowserTextureSampleQueue, ESPMode::ThreadSafe> WebBrowserTextureSamplesQueue;
+
+	/** Texture sample object pool. */
+	FWebBrowserTextureSamplePool* TextureSamplePool;
+
+	// Handle to detect changes in valid area to move the floating close button */
+	FDelegateHandle OnSafeFrameChangedEventHandle;
+#endif
+	
+	CGRect WindowFrame;
+};
+
+
+
+
+@implementation AppleWebViewWrapper
+
+#if !PLATFORM_TVOS
+@synthesize WebView;
+#if PLATFORM_IOS
+@synthesize CloseButton;
+#endif
+@synthesize WebViewContainer;
+#endif
+@synthesize NextURL;
+@synthesize NextContent;
+
+-(void)create:(TSharedPtr<SAppleWebBrowserWidget>)InWebBrowserWidget userAgentApplication: (NSString*)UserAgentApplication useTransparency : (bool)InUseTransparency
+supportsMetal : (bool)InSupportsMetal supportsMetalMRT : (bool)InSupportsMetalMRT enableFloatingCloseButton : (bool)bEnableFloatingCloseButton;
+{
+	WebBrowserWidget = InWebBrowserWidget;
+	NextURL = nil;
+	NextContent = nil;
+	VideoTexture = nil;
+	bNeedsAddToView = true;
+	Is3DBrowser = false;
+	bVideoTextureValid = false;
+	bSupportsMetal = InSupportsMetal;
+	bSupportsMetalMRT = InSupportsMetalMRT;
+
+#if !PLATFORM_TVOS
+	dispatch_async(dispatch_get_main_queue(), ^
+	{
+		#if PLATFORM_IOS
+		WebViewContainer = [[UIView alloc]initWithFrame:CGRectMake(1, 1, 100, 100)];
+		[self.WebViewContainer setOpaque : NO];
+		[self.WebViewContainer setBackgroundColor : [UIColor clearColor]];
+		#else
+		WebViewContainer = [[NSView alloc]initWithFrame:CGRectMake(1, 1, 100, 100)];
+		#endif
+
+		WKWebViewConfiguration* theConfiguration = [[WKWebViewConfiguration alloc] init];
+		
+		theConfiguration.websiteDataStore = [WKWebsiteDataStore dataStoreForIdentifier: FAppleCookieManager::CookieManagerIdentifier];
+		
+		[theConfiguration.preferences setValue:@YES forKey:@"allowFileAccessFromFileURLs"];
+		[theConfiguration.preferences setValue:@YES forKey:@"developerExtrasEnabled"];
+		
+		NSString* MessageHandlerName = [NSString stringWithFString : FMobileJSScripting::JSMessageHandler];
+		theConfiguration.applicationNameForUserAgent = UserAgentApplication;
+		[theConfiguration.userContentController addScriptMessageHandler:self name: MessageHandlerName];
+
+		WebView = [[WKWebView alloc]initWithFrame:CGRectMake(1, 1, 100, 100)configuration : theConfiguration];
+#if !UE_BUILD_SHIPPING
+		WebView.inspectable = CVarWebViewEnableInspectable.GetValueOnAnyThread()? YES : NO;
+#endif
+		[self.WebViewContainer addSubview : WebView];
+		WebView.navigationDelegate = self;
+		WebView.UIDelegate = self;
+
+#if PLATFORM_IOS
+		WebView.scrollView.bounces = NO;
+		
+		if (InUseTransparency)
+		{
+			[self.WebView setOpaque : NO];
+			[self.WebView setBackgroundColor : [UIColor clearColor]];
+		}
+		else
+		{
+			[self.WebView setOpaque : YES];
+		}
+		
+		if (bEnableFloatingCloseButton)
+		{
+			CloseButton = MakeCloseButton();
+			[self.WebView addSubview : CloseButton];
+
+			TWeakPtr<SAppleWebBrowserWidget> WeakWebBrowserWidget = InWebBrowserWidget;
+			CloseButton.TapHandler = [^{
+				if (TSharedPtr<SAppleWebBrowserWidget> StrongWebBrowserWidget = WeakWebBrowserWidget.Pin())
+				{
+					StrongWebBrowserWidget->FloatingCloseButtonPressed();
+				}
+			} copy];
+		}
+#endif
+
+		[theConfiguration release];
+		[self setDefaultVisibility];
+	});
+#endif
+}
+
+-(void)didRotate;
+{
+#if !PLATFORM_TVOS && !PLATFORM_MAC
+	dispatch_async(dispatch_get_main_queue(), ^ {
+		if (CloseButton)
+		{
+			[CloseButton setupLayout];
+		}
+	});
+#endif
+}
+
+-(void)close;
+{
+#if !PLATFORM_TVOS
+	WebView.navigationDelegate = nil;
+	dispatch_async(dispatch_get_main_queue(), ^
+	{
+		[self.WebViewContainer removeFromSuperview];
+		[self.WebView removeFromSuperview];
+		
+		#if PLATFORM_IOS
+		if (CloseButton)
+		{
+			[CloseButton release];
+		}
+		CloseButton = nil;
+		#endif
+		[WebView release];
+		[WebViewContainer release];
+
+		WebView = nil;
+		WebViewContainer = nil;
+	});
+#endif
+}
+
+-(void)showFloatingCloseButton:(BOOL)bShow setDraggable:(BOOL)bDraggable;
+{
+#if !PLATFORM_TVOS && !PLATFORM_MAC
+	dispatch_async(dispatch_get_main_queue(), ^
+	{
+		if (CloseButton)
+		{
+			[CloseButton showButton:bShow setDraggable: bDraggable];
+		}
+		else
+		{
+			UE_LOG(LogAppleWebBrowser, Warning, TEXT("[PlatformWebBrowser]: Close button not enabled in config. Add config for Engine:[Browser]bEnableFloatingCloseButton=true"));
+		}
+	});
+#endif
+}
+
+-(void)dealloc;
+{
+#if !PLATFORM_TVOS
+	if (WebView != nil)
+	{
+		WebView.navigationDelegate = nil;
+		[WebView release];
+		WebView = nil;
+	};
+
+#if !PLATFORM_MAC
+	if (CloseButton)
+	{
+		[CloseButton release];
+		CloseButton = nil;
+	};
+#endif
+
+	[WebViewContainer release];
+	WebViewContainer = nil;
+#endif
+	[NextContent release];
+	NextContent = nil;
+	[NextURL release];
+	NextURL = nil;
+
+	[super dealloc];
+}
+
+-(void)updateframe:(CGRect)InFrame window:(FGenericWindow*)EnclosedWindow;
+{
+	self.DesiredFrame = InFrame;
+
+#if !PLATFORM_TVOS
+	dispatch_async(dispatch_get_main_queue(), ^
+	{
+		if (WebView != nil)
+		{
+			WebViewContainer.frame = self.DesiredFrame;
+			WebView.frame = WebViewContainer.bounds;
+			if (bNeedsAddToView)
+			{
+				bNeedsAddToView = false;
+				
+				#if PLATFORM_IOS 
+				[self.WebViewContainer setOpaque : NO];
+				[self.WebViewContainer setBackgroundColor : [UIColor clearColor]];
+				[[IOSAppDelegate GetDelegate].IOSView addSubview : WebViewContainer];
+				
+				if (CloseButton != nil)
+				{
+					[CloseButton setupLayout];
+				}
+				#else
+				// Are we adding our WebViewContainer to the proper Window here?
+				FCocoaWindow *CocoaWindow = ((FMacWindow*)EnclosedWindow)->GetWindowHandle();
+				NSView *CocoaWindowView = CocoaWindow.openGLView;
+				
+				[CocoaWindowView addSubview : WebViewContainer];
+				#endif
+			}
+			else
+			{
+				if (NextContent != nil)
+				{
+					// Load web content from string
+					[self.WebView loadHTMLString : NextContent baseURL : NextURL];
+					NextContent = nil;
+					NextURL = nil;
+				}
+				else
+					if (NextURL != nil)
+					{
+						// Load web content from URL
+						NSURLRequest *nsrequest = [NSURLRequest requestWithURL : NextURL];
+						[self.WebView loadRequest : nsrequest];
+						NextURL = nil;
+					}
+			}
+		}
+	});
+#endif
+}
+
+-(NSString *)UrlDecode:(NSString *)stringToDecode
+{
+	NSString *result = [stringToDecode stringByReplacingOccurrencesOfString : @"+" withString:@" "];
+	result = [result stringByRemovingPercentEncoding];
+	return result;
+}
+
+#if !PLATFORM_TVOS
+-(void)userContentController:(WKUserContentController *)userContentController didReceiveScriptMessage : (WKScriptMessage *)message
+{
+	if ([message.body isKindOfClass : [NSString class]])
+	{
+		NSString *Message = message.body;
+		if (WebBrowserWidget != nil && Message != nil)
+		{
+			//NSLog(@"Received message %@", Message);
+			WebBrowserWidget->ProcessScriptMessage(Message);
+		}
+
+	}
+}
+#endif
+
+-(void)executejavascript:(NSString*)InJavaScript
+{
+#if !PLATFORM_TVOS
+	dispatch_async(dispatch_get_main_queue(), ^
+	{
+		const char* InJavaScriptCString = [InJavaScript UTF8String];
+		FString InJavaScriptFString = FString(InJavaScriptCString);
+		
+		//NSLog(@"executejavascript %@", InJavaScript);
+		
+		[self.WebView evaluateJavaScript : InJavaScript completionHandler : ^(id Object, NSError* Error) {
+			if(Error != nil) {
+				UE_LOG(LogAppleWebBrowser, Error, TEXT("executeJavaScript(%s) errored with NSError %d!"), 
+					   *InJavaScriptFString, Error.code);
+			}
+		}];
+	});
+#endif
+}
+
+-(void)loadurl:(NSURL*)InURL;
+{
+	dispatch_async(dispatch_get_main_queue(), ^
+	{
+		self.NextURL = InURL;
+	});
+}
+
+-(void)loadstring:(NSString*)InString dummyurl : (NSURL*)InURL;
+{
+	dispatch_async(dispatch_get_main_queue(), ^
+	{
+		self.NextContent = InString;
+		self.NextURL = InURL;
+	});
+}
+
+-(void)set3D:(bool)InIs3DBrowser;
+{
+	dispatch_async(dispatch_get_main_queue(), ^
+	{
+		if (Is3DBrowser != InIs3DBrowser)
+		{
+			//default is 2D
+			Is3DBrowser = InIs3DBrowser;
+			[self setDefaultVisibility];
+		}
+	});
+}
+
+-(void)setDefaultVisibility;
+{
+#if !PLATFORM_TVOS
+	dispatch_async(dispatch_get_main_queue(), ^
+	{
+		if (Is3DBrowser)
+		{
+			[self.WebViewContainer setHidden : YES];
+		}
+		else
+		{
+			[self.WebViewContainer setHidden : NO];
+		}
+	});
+#endif
+}
+
+-(void)setVisibility:(bool)InIsVisible;
+{
+#if !PLATFORM_TVOS
+	dispatch_async(dispatch_get_main_queue(), ^
+	{
+		if (InIsVisible)
+		{
+			[self setDefaultVisibility];
+		}
+		else
+		{
+			[self.WebViewContainer setHidden : YES];
+		}
+	});
+#endif
+}
+
+-(void)stopLoading;
+{
+#if !PLATFORM_TVOS
+	dispatch_async(dispatch_get_main_queue(), ^
+	{
+		[self.WebView stopLoading];
+	});
+#endif
+}
+
+-(void)reload;
+{
+#if !PLATFORM_TVOS
+	dispatch_async(dispatch_get_main_queue(), ^
+	{
+		[self.WebView reload];
+	});
+#endif
+}
+
+-(void)goBack;
+{
+#if !PLATFORM_TVOS
+	dispatch_async(dispatch_get_main_queue(), ^
+	{
+		[self.WebView goBack];
+	});
+#endif
+}
+
+-(void)goForward;
+{
+#if !PLATFORM_TVOS
+	dispatch_async(dispatch_get_main_queue(), ^
+	{
+		[self.WebView goForward];
+	});
+#endif
+}
+
+-(bool)canGoBack;
+{
+#if PLATFORM_TVOS
+	return false;
+#else
+	return [self.WebView canGoBack];
+#endif
+}
+
+-(bool)canGoForward;
+{
+#if PLATFORM_TVOS
+	return false;
+#else
+	return [self.WebView canGoForward];
+#endif
+}
+
+-(FTextureRHIRef)GetVideoTexture;
+{
+	return VideoTexture;
+}
+
+-(void)SetVideoTexture:(FTextureRHIRef)Texture;
+{
+	VideoTexture = Texture;
+}
+
+-(void)SetVideoTextureValid:(bool)Condition;
+{
+	bVideoTextureValid = Condition;
+}
+
+-(bool)IsVideoTextureValid;
+{
+	return bVideoTextureValid;
+}
+
+-(bool)UpdateVideoFrame:(void*)ptr;
+{
+#if !PLATFORM_TVOS
+	@synchronized(self) // Briefly block render thread
+	{
+		id<MTLTexture> ptrToMetalTexture = (id<MTLTexture>)ptr;
+		NSUInteger width = [ptrToMetalTexture width];
+		NSUInteger height = [ptrToMetalTexture height];
+
+		[self updateWebViewMetalTexture : ptrToMetalTexture];
+	}
+#endif
+	return true;
+}
+
+-(void)updateWebViewMetalTexture:(id<MTLTexture>)texture
+{
+#if !PLATFORM_TVOS
+	@autoreleasepool {
+		
+		CGColorSpaceRef ColorSpace = CGColorSpaceCreateDeviceRGB();
+		NSUInteger Width = [texture width];
+		NSUInteger Height = [texture height];
+		CGContextRef WebViewContext = CGBitmapContextCreate(NULL, Width, Height, 8, 4 * Width, ColorSpace, (CGBitmapInfo)kCGImageAlphaPremultipliedLast);
+		
+#if PLATFORM_MAC
+		WebView.wantsLayer = YES;
+		if(WebView.layer == nil) {
+			return;
+		}
+		[WebView.layer renderInContext: WebViewContext];
+#else
+		UIGraphicsBeginImageContextWithOptions(WebView.frame.size, NO, 1.0f);
+	    [WebView drawViewHierarchyInRect : WebView.bounds afterScreenUpdates : NO];
+		UIImage *Image = UIGraphicsGetImageFromCurrentImageContext();
+		UIGraphicsEndImageContext();
+		CGImage *WebViewCGImage = Image.CGImage;
+		CGContextDrawImage(WebViewContext, CGRectMake(0, 0, Width, Height), WebViewCGImage);
+#endif
+		[texture replaceRegion : MTLRegionMake2D(0, 0, Width, Height)
+			mipmapLevel : 0
+			withBytes : CGBitmapContextGetData(WebViewContext)
+			bytesPerRow : 4 * Width];
+		CGColorSpaceRelease(ColorSpace);
+		CGContextRelease(WebViewContext);
+		
+#if !PLATFORM_MAC
+		Image = nil;
+#endif
+	}
+#endif
+}
+
+#if !PLATFORM_TVOS
+
+- (nullable WKWebView *)webView:(WKWebView *)InWebView createWebViewWithConfiguration:(WKWebViewConfiguration *)InConfiguration forNavigationAction:(WKNavigationAction *)InNavigationAction windowFeatures:(WKWindowFeatures *)InWindowFeatures
+{
+	NSURLRequest *request = InNavigationAction.request;
+	FString UrlStr([[request URL] absoluteString]);
+
+	if (InNavigationAction.targetFrame == nil && !UrlStr.IsEmpty() && FPlatformProcess::CanLaunchURL(*UrlStr))
+	{
+		if (WebBrowserWidget->HandleOnBeforePopup(UrlStr, TEXT("_blank")))
+		{
+			// Launched the URL in external browser, don't create a new webview
+			return nil;
+		}
+	}
+	return nil;
+}
+
+- (void)webView:(WKWebView*)InWebView decidePolicyForNavigationAction : (WKNavigationAction*)InNavigationAction decisionHandler : (void(^)(WKNavigationActionPolicy))InDecisionHandler
+{
+	NSURLRequest *request = InNavigationAction.request;
+	FString UrlStr([[request URL] absoluteString]);
+
+	if (InNavigationAction.targetFrame == nil && !UrlStr.IsEmpty() && FPlatformProcess::CanLaunchURL(*UrlStr))
+	{
+		if (WebBrowserWidget->HandleOnBeforePopup(UrlStr, TEXT("_blank")))
+		{
+			// Launched the URL in external browser, don't open the link here too
+			InDecisionHandler(WKNavigationActionPolicyCancel);
+			return;
+		}
+	}
+	
+	WebBrowserWidget->HandleShouldOverrideUrlLoading(UrlStr);
+	InDecisionHandler(WKNavigationActionPolicyAllow);
+}
+
+-(void)webView:(WKWebView *)InWebView didCommitNavigation : (WKNavigation *)InNavigation
+{
+	NSString* CurrentUrl = [self.WebView URL].absoluteString;
+	NSString* Title = [self.WebView title];
+	
+//	NSLog(@"didCommitNavigation: %@", CurrentUrl);
+	WebBrowserWidget->HandleReceivedTitle(Title);
+	WebBrowserWidget->HandlePageLoad(CurrentUrl, true);
+}
+
+-(void)webView:(WKWebView *)InWebView didFinishNavigation : (WKNavigation *)InNavigation
+{
+	NSString* CurrentUrl = [self.WebView URL].absoluteString;
+	NSString* Title = [self.WebView title];
+	// NSLog(@"didFinishNavigation: %@", CurrentUrl);
+	WebBrowserWidget->HandleReceivedTitle(Title);
+	WebBrowserWidget->HandlePageLoad(CurrentUrl, false);
+}
+-(void)webView:(WKWebView *)InWebView didFailNavigation : (WKNavigation *)InNavigation withError : (NSError*)InError
+{
+	if (InError.domain == NSURLErrorDomain && InError.code == NSURLErrorCancelled)
+	{
+		//ignore this one, interrupted load
+		return;
+	}
+	NSString* CurrentUrl = [InError.userInfo objectForKey : @"NSErrorFailingURLStringKey"];
+//	NSLog(@"didFailNavigation: %@, error %@", CurrentUrl, InError);
+	WebBrowserWidget->HandleReceivedError(InError.code, CurrentUrl);
+}
+-(void)webView:(WKWebView *)InWebView didFailProvisionalNavigation : (WKNavigation *)InNavigation withError : (NSError*)InError
+{
+	NSString* CurrentUrl = [InError.userInfo objectForKey : @"NSErrorFailingURLStringKey"];
+	// NSLog(@"didFailProvisionalNavigation: %@, error %@", CurrentUrl, InError);
+	WebBrowserWidget->HandleReceivedError(InError.code, CurrentUrl);
+}
+#endif
+@end
+
+namespace {
+	static const FString JSGetSourceCommand = TEXT("GetSource");
+	static const FString JSMessageGetSourceScript =
+		TEXT("	window.webkit.messageHandlers.") + FMobileJSScripting::JSMessageHandler + TEXT(".postMessage('")+ JSGetSourceCommand +
+		TEXT("/' + encodeURIComponent(document.documentElement.innerHTML));");
+
+}
+
+FWebBrowserWindow::FWebBrowserWindow(FString InUrl, TOptional<FString> InContentsToLoad, bool InShowErrorMessage, bool InThumbMouseButtonNavigation, bool InUseTransparency, bool bInJSBindingToLoweringEnabled, FString InUserAgentApplication)
+	: CurrentUrl(MoveTemp(InUrl))
+	, UserAgentApplication(MoveTemp(InUserAgentApplication))
+	, ContentsToLoad(MoveTemp(InContentsToLoad))
+	, bUseTransparency(InUseTransparency)
+	, DocumentState(EWebBrowserDocumentState::NoDocument)
+	, ErrorCode(0)
+	, WindowSize(FIntPoint(500, 500))
+	, bIsDisabled(false)
+	, bIsVisible(true)
+	, bTickedLastFrame(true)
+{
+	bool bInjectJSOnPageStarted = false;
+	GConfig->GetBool(TEXT("Browser"), TEXT("bInjectJSOnPageStarted"), bInjectJSOnPageStarted, GEngineIni);
+	Scripting = MakeShared<FMobileJSScripting>(bInJSBindingToLoweringEnabled, bInjectJSOnPageStarted);
+}
+
+FWebBrowserWindow::~FWebBrowserWindow()
+{
+	CloseBrowser(true, false);
+}
+
+void FWebBrowserWindow::LoadURL(FString NewURL)
+{
+	BrowserWidget->LoadURL(NewURL);
+}
+
+void FWebBrowserWindow::LoadString(FString Contents, FString DummyURL)
+{
+	BrowserWidget->LoadString(Contents, DummyURL);
+}
+
+TSharedRef<SWidget> FWebBrowserWindow::CreateWidget()
+{
+	TSharedRef<SAppleWebBrowserWidget> BrowserWidgetRef =
+		SNew(SAppleWebBrowserWidget)
+		.UseTransparency(bUseTransparency)
+		.InitialURL(CurrentUrl)
+		.UserAgentApplication(UserAgentApplication)
+		.WebBrowserWindow(SharedThis(this));
+
+	BrowserWidget = BrowserWidgetRef;
+
+	Scripting->SetWindow(SharedThis(this));
+	
+	return BrowserWidgetRef;
+}
+
+void FWebBrowserWindow::SetViewportSize(FIntPoint ViewportSize, FIntPoint ViewportPos)
+{
+	WindowSize = ViewportSize;
+	if(BrowserWidget != nullptr) 
+	{
+		BrowserWidget->UpdateFrameWithViewport(WindowSize, ViewportPos);
+	}
+}
+
+FIntPoint FWebBrowserWindow::GetViewportSize() const
+{
+	return WindowSize;
+}
+
+FSlateShaderResource* FWebBrowserWindow::GetTexture(bool bIsPopup /*= false*/)
+{
+	return nullptr;
+}
+
+bool FWebBrowserWindow::IsValid() const
+{
+	return false;
+}
+
+bool FWebBrowserWindow::IsInitialized() const
+{
+	return true;
+}
+
+bool FWebBrowserWindow::IsClosing() const
+{
+	return false;
+}
+
+EWebBrowserDocumentState FWebBrowserWindow::GetDocumentLoadingState() const
+{
+	return DocumentState;
+}
+
+FString FWebBrowserWindow::GetTitle() const
+{
+	return Title;
+}
+
+FString FWebBrowserWindow::GetUrl() const
+{
+	return CurrentUrl;
+}
+
+bool FWebBrowserWindow::OnKeyDown(const FKeyEvent& InKeyEvent)
+{
+	return false;
+}
+
+bool FWebBrowserWindow::OnKeyUp(const FKeyEvent& InKeyEvent)
+{
+	return false;
+}
+
+bool FWebBrowserWindow::OnKeyChar(const FCharacterEvent& InCharacterEvent)
+{
+	return false;
+}
+
+FReply FWebBrowserWindow::OnMouseButtonDown(const FGeometry& MyGeometry, const FPointerEvent& MouseEvent, bool bIsPopup)
+{
+	return FReply::Unhandled();
+}
+
+FReply FWebBrowserWindow::OnMouseButtonUp(const FGeometry& MyGeometry, const FPointerEvent& MouseEvent, bool bIsPopup)
+{
+	return FReply::Unhandled();
+}
+
+FReply FWebBrowserWindow::OnMouseButtonDoubleClick(const FGeometry& MyGeometry, const FPointerEvent& MouseEvent, bool bIsPopup)
+{
+	return FReply::Unhandled();
+}
+
+FReply FWebBrowserWindow::OnMouseMove(const FGeometry& MyGeometry, const FPointerEvent& MouseEvent, bool bIsPopup)
+{
+	return FReply::Unhandled();
+}
+
+void FWebBrowserWindow::OnMouseLeave(const FPointerEvent& MouseEvent)
+{
+}
+
+void FWebBrowserWindow::SetSupportsMouseWheel(bool bValue)
+{
+
+}
+
+bool FWebBrowserWindow::GetSupportsMouseWheel() const
+{
+	return false;
+}
+
+FReply FWebBrowserWindow::OnMouseWheel(const FGeometry& MyGeometry, const FPointerEvent& MouseEvent, bool bIsPopup)
+{
+	return FReply::Unhandled();
+}
+
+FReply FWebBrowserWindow::OnTouchGesture(const FGeometry& MyGeometry, const FPointerEvent& GestureEvent, bool bIsPopup)
+{
+	return FReply::Unhandled();
+}
+
+
+void FWebBrowserWindow::OnFocus(bool SetFocus, bool bIsPopup)
+{
+}
+
+void FWebBrowserWindow::OnCaptureLost()
+{
+}
+
+bool FWebBrowserWindow::CanGoBack() const
+{
+	return BrowserWidget->CanGoBack();
+}
+
+void FWebBrowserWindow::GoBack()
+{
+	BrowserWidget->GoBack();
+}
+
+bool FWebBrowserWindow::CanGoForward() const
+{
+	return BrowserWidget->CanGoForward();
+}
+
+void FWebBrowserWindow::GoForward()
+{
+	BrowserWidget->GoForward();
+}
+
+bool FWebBrowserWindow::IsLoading() const
+{
+	return DocumentState != EWebBrowserDocumentState::Loading;
+}
+
+void FWebBrowserWindow::Reload()
+{
+	BrowserWidget->Reload();
+}
+
+void FWebBrowserWindow::StopLoad()
+{
+	BrowserWidget->StopLoad();
+}
+
+void FWebBrowserWindow::GetSource(TFunction<void(const FString&)> Callback) const
+{
+	//@todo: decide what to do about multiple pending requests
+	GetPageSourceCallback.Emplace(Callback);
+
+	// Ugly hack: Work around the fact that ExecuteJavascript is non-const.
+	const_cast<FWebBrowserWindow*>(this)->ExecuteJavascript(JSMessageGetSourceScript);
+}
+
+int FWebBrowserWindow::GetLoadError()
+{
+	return ErrorCode;
+}
+
+void FWebBrowserWindow::NotifyDocumentError(const FString& InCurrentUrl, int InErrorCode)
+{
+	if (!CurrentUrl.Equals(InCurrentUrl, ESearchCase::CaseSensitive))
+	{
+		CurrentUrl = InCurrentUrl;
+		UrlChangedEvent.Broadcast(CurrentUrl);
+	}
+
+	ErrorCode = InErrorCode;
+	DocumentState = EWebBrowserDocumentState::Error;
+	DocumentStateChangedEvent.Broadcast(DocumentState);
+}
+
+void FWebBrowserWindow::NotifyDocumentLoadingStateChange(const FString& InCurrentUrl, bool IsLoading)
+{
+	// Ignore a load completed notification if there was an error.
+	// For load started, reset any errors from previous page load.
+	bool bIsNotError = DocumentState != EWebBrowserDocumentState::Error;
+	if (IsLoading || bIsNotError)
+	{
+		if (!CurrentUrl.Equals(InCurrentUrl, ESearchCase::CaseSensitive))
+		{
+			CurrentUrl = InCurrentUrl;
+			UrlChangedEvent.Broadcast(CurrentUrl);
+		}
+
+		if (bIsNotError && !InCurrentUrl.StartsWith("javascript:"))
+		{
+			if (IsLoading)
+			{
+				Scripting->PageStarted(SharedThis(this));
+			}
+			else
+			{
+				Scripting->PageLoaded(SharedThis(this));
+			}
+		}
+		ErrorCode = 0;
+		DocumentState = IsLoading
+			? EWebBrowserDocumentState::Loading
+			: EWebBrowserDocumentState::Completed;
+		DocumentStateChangedEvent.Broadcast(DocumentState);
+	}
+
+}
+
+void FWebBrowserWindow::SetIsDisabled(bool bValue)
+{
+	bIsDisabled = bValue;
+}
+
+TSharedPtr<SWindow> FWebBrowserWindow::GetParentWindow() const
+{
+	return ParentWindow;
+}
+
+void FWebBrowserWindow::SetParentWindow(TSharedPtr<SWindow> Window)
+{
+	ParentWindow = Window;
+}
+
+void FWebBrowserWindow::ShowFloatingCloseButton(bool bShow, bool bDraggable)
+{
+	if (BrowserWidget)
+	{
+		BrowserWidget->ShowFloatingCloseButton(bShow, bDraggable);
+	}
+}
+
+void FWebBrowserWindow::ExecuteJavascript(const FString& Script)
+{
+	BrowserWidget->ExecuteJavascript(Script);
+}
+
+void FWebBrowserWindow::CloseBrowser(bool bForce, bool bBlockTillClosed /* ignored */)
+{
+	BrowserWidget->Close();
+}
+
+bool FWebBrowserWindow::OnJsMessageReceived(const FString& Command, const TArray<FString>& Params, const FString& Origin)
+{
+	if (Command.Equals(JSGetSourceCommand, ESearchCase::CaseSensitive) && GetPageSourceCallback.IsSet() && Params.Num() == 1)
+	{
+		GetPageSourceCallback.GetValue()(Params[0]);
+		GetPageSourceCallback.Reset();
+		return true;
+	}
+	return Scripting->OnJsMessageReceived(Command, Params, Origin);
+}
+
+void FWebBrowserWindow::BindUObject(const FString& Name, UObject* Object, bool bIsPermanent /*= true*/)
+{
+	Scripting->BindUObject(Name, Object, bIsPermanent);
+}
+
+void FWebBrowserWindow::UnbindUObject(const FString& Name, UObject* Object /*= nullptr*/, bool bIsPermanent /*= true*/)
+{
+	Scripting->UnbindUObject(Name, Object, bIsPermanent);
+}
+
+void FWebBrowserWindow::CheckTickActivity()
+{
+	if (bIsVisible != bTickedLastFrame)
+	{
+		bIsVisible = bTickedLastFrame;
+		BrowserWidget->SetWebBrowserVisibility(bIsVisible);
+	}
+
+	bTickedLastFrame = false;
+}
+
+void FWebBrowserWindow::SetTickLastFrame()
+{
+	bTickedLastFrame = !bIsDisabled;
+}
+
+bool FWebBrowserWindow::IsVisible()
+{
+	return bIsVisible;
+}
+
+void FWebBrowserWindow::FloatingCloseButtonPressed()
+{
+	if (OnFloatingCloseButtonPressed().IsBound())
+	{
+		OnFloatingCloseButtonPressed().Execute();
+	}
+}
+#endif
diff --git a/Engine/Source/Runtime/WebBrowser/Private/Apple/ApplePlatformWebBrowser.h b/Engine/Source/Runtime/WebBrowser/Private/Apple/ApplePlatformWebBrowser.h
new file mode 100644
--- /dev/null
+++ b/Engine/Source/Runtime/WebBrowser/Private/Apple/ApplePlatformWebBrowser.h
@@ -0,0 +1,452 @@
+// Copyright Epic Games, Inc. All Rights Reserved.
+
+#pragma once
+
+#include "CoreMinimal.h"
+
+#if !WITH_CEF3 && (PLATFORM_IOS || PLATFORM_MAC)
+
+#include "IWebBrowserWindow.h"
+#include "Widgets/SWindow.h"
+
+#include "Engine/Texture2D.h"
+#include "Materials/MaterialInstanceDynamic.h"
+#include "WebBrowserTexture.h"
+#include "Misc/ConfigCacheIni.h"
+
+#include "MobileJS/MobileJSScripting.h"
+
+#if PLATFORM_IOS
+#import <UIKit/UIKit.h>
+#else
+#import <AppKit/AppKit.h>
+#endif
+
+#if !PLATFORM_TVOS
+#import "WebKit/WebKit.h"
+#endif
+
+
+DECLARE_LOG_CATEGORY_EXTERN(LogAppleWebBrowser, Log, All);
+
+class SAppleWebBrowserWidget;
+class SWebBrowserView;
+
+/**
+* Wrapper to contain the UIWebView and implement its delegate functions
+*/
+#if !PLATFORM_TVOS
+@class WebViewCloseButton;
+
+@interface AppleWebViewWrapper : NSObject <WKUIDelegate, WKNavigationDelegate, WKScriptMessageHandler>
+#else
+@interface AppleWebViewWrapper : NSObject
+#endif
+{
+	TSharedPtr<SAppleWebBrowserWidget> WebBrowserWidget;
+	FTextureRHIRef VideoTexture;
+	bool bNeedsAddToView;
+	bool Is3DBrowser;
+	bool bVideoTextureValid;
+	bool bSupportsMetal;
+	bool bSupportsMetalMRT;
+}
+#if !PLATFORM_TVOS
+@property(strong) WKWebView* WebView;
+#if PLATFORM_MAC
+@property(strong) NSView* WebViewContainer;
+#else
+@property(strong) WebViewCloseButton* CloseButton;
+@property(strong) UIView* WebViewContainer;
+#endif
+#endif
+@property(copy) NSURL* NextURL;
+@property(copy) NSString* NextContent;
+@property CGRect DesiredFrame;
+
+-(void)create:(TSharedPtr<SAppleWebBrowserWidget>)InWebBrowserWidget userAgentApplication: (NSString*)UserAgentApplication useTransparency : (bool)InUseTransparency
+supportsMetal : (bool)InSupportsMetal supportsMetalMRT : (bool)InSupportsMetalMRT enableFloatingCloseButton : (bool)bEnableFloatingCloseButton;
+-(void)didRotate;
+-(void)close;
+-(void)showFloatingCloseButton:(BOOL)bShow setDraggable:(BOOL)bDraggable;
+-(void)dealloc;
+-(void)updateframe:(CGRect)InFrame window:(FGenericWindow*)EnclosedWindow;
+-(void)loadstring:(NSString*)InString dummyurl:(NSURL*)InURL;
+-(void)loadurl:(NSURL*)InURL;
+-(void)executejavascript:(NSString*)InJavaScript;
+-(void)set3D:(bool)InIs3DBrowser;
+-(void)setDefaultVisibility;
+-(void)setVisibility:(bool)InIsVisible;
+-(void)stopLoading;
+-(void)reload;
+-(void)goBack;
+-(void)goForward;
+-(bool)canGoBack;
+-(bool)canGoForward;
+-(FTextureRHIRef)GetVideoTexture;
+-(void)SetVideoTexture:(FTextureRHIRef)Texture;
+-(void)SetVideoTextureValid:(bool)Condition;
+-(bool)IsVideoTextureValid;
+-(bool)UpdateVideoFrame:(void*)ptr;
+-(void)updateWebViewMetalTexture : (id<MTLTexture>)texture;
+#if !PLATFORM_TVOS
+-(void)webView:(WKWebView*)InWebView decidePolicyForNavigationAction : (WKNavigationAction*)InNavigationAction decisionHandler : (void(^)(WKNavigationActionPolicy))InDecisionHandler;
+-(void)webView:(WKWebView*)InWebView didCommitNavigation : (WKNavigation*)InNavigation;
+#endif
+@end
+
+/**
+* Implementation of interface for dealing with a Web Browser window.
+*/
+class FWebBrowserWindow
+	: public IWebBrowserWindow
+	, public TSharedFromThis<FWebBrowserWindow>
+{
+	// The WebBrowserSingleton should be the only one creating instances of this class
+	friend class FWebBrowserSingleton;
+	// CreateWidget should only be called by the WebBrowserView
+	friend class SWebBrowserView;
+
+private:
+	/**
+	* Creates and initializes a new instance.
+	*
+	* @param InUrl The Initial URL that will be loaded.
+	* @param InContentsToLoad Optional string to load as a web page.
+	* @param InShowErrorMessage Whether to show an error message in case of loading errors.
+	*/
+	FWebBrowserWindow(FString InUrl, TOptional<FString> InContentsToLoad, bool ShowErrorMessage, bool bThumbMouseButtonNavigation, bool bUseTransparency, bool bInJSBindingToLoweringEnabled, FString InUserAgentApplication);
+
+	/**
+	 * Create the SWidget for this WebBrowserWindow
+	 */
+	TSharedRef<SWidget> CreateWidget();
+
+public:
+	/** Virtual Destructor. */
+	virtual ~FWebBrowserWindow();
+
+public:
+
+	// IWebBrowserWindow Interface
+
+	virtual void LoadURL(FString NewURL) override;
+	virtual void LoadString(FString Contents, FString DummyURL) override;
+	virtual void SetViewportSize(FIntPoint ViewportWindowSize, FIntPoint WindowPos) override;
+	virtual FIntPoint GetViewportSize() const override;
+	virtual FSlateShaderResource* GetTexture(bool bIsPopup = false) override;
+	virtual bool IsValid() const override;
+	virtual bool IsInitialized() const override;
+	virtual bool IsClosing() const override;
+	virtual EWebBrowserDocumentState GetDocumentLoadingState() const override;
+	virtual FString GetTitle() const override;
+	virtual FString GetUrl() const override;
+	virtual bool OnKeyDown(const FKeyEvent& InKeyEvent) override;
+	virtual bool OnKeyUp(const FKeyEvent& InKeyEvent) override;
+	virtual bool OnKeyChar(const FCharacterEvent& InCharacterEvent) override;
+	virtual FReply OnMouseButtonDown(const FGeometry& MyGeometry, const FPointerEvent& MouseEvent, bool bIsPopup) override;
+	virtual FReply OnMouseButtonUp(const FGeometry& MyGeometry, const FPointerEvent& MouseEvent, bool bIsPopup) override;
+	virtual FReply OnMouseButtonDoubleClick(const FGeometry& MyGeometry, const FPointerEvent& MouseEvent, bool bIsPopup) override;
+	virtual FReply OnMouseMove(const FGeometry& MyGeometry, const FPointerEvent& MouseEvent, bool bIsPopup) override;
+	virtual void OnMouseLeave(const FPointerEvent& MouseEvent) override;
+	virtual void SetSupportsMouseWheel(bool bValue) override;
+	virtual bool GetSupportsMouseWheel() const override;
+	virtual FReply OnMouseWheel(const FGeometry& MyGeometry, const FPointerEvent& MouseEvent, bool bIsPopup) override;
+	virtual FReply OnTouchGesture(const FGeometry& MyGeometry, const FPointerEvent& GestureEvent, bool bIsPopup) override;
+	virtual void OnFocus(bool SetFocus, bool bIsPopup) override;
+	virtual void OnCaptureLost() override;
+	virtual bool CanGoBack() const override;
+	virtual void GoBack() override;
+	virtual bool CanGoForward() const override;
+	virtual void GoForward() override;
+	virtual bool IsLoading() const override;
+	virtual void Reload() override;
+	virtual void StopLoad() override;
+	virtual void ExecuteJavascript(const FString& Script) override;
+	virtual void CloseBrowser(bool bForce, bool bBlockTillClosed) override;
+	virtual void BindUObject(const FString& Name, UObject* Object, bool bIsPermanent = true) override;
+	virtual void UnbindUObject(const FString& Name, UObject* Object = nullptr, bool bIsPermanent = true) override;
+	virtual void GetSource(TFunction<void (const FString&)> Callback) const;
+	virtual int GetLoadError() override;
+	virtual void SetIsDisabled(bool bValue) override;
+	virtual TSharedPtr<SWindow> GetParentWindow() const override;
+	virtual void SetParentWindow(TSharedPtr<SWindow> Window) override;
+	virtual void ShowFloatingCloseButton(bool bShow, bool bDraggable) override;
+
+	// TODO: None of these events are actually called
+
+	DECLARE_DERIVED_EVENT(FWebBrowserWindow, IWebBrowserWindow::FOnDocumentStateChanged, FOnDocumentStateChanged);
+	virtual FOnDocumentStateChanged& OnDocumentStateChanged() override
+	{
+		return DocumentStateChangedEvent;
+	}
+
+	DECLARE_DERIVED_EVENT(FWebBrowserWindow, IWebBrowserWindow::FOnTitleChanged, FOnTitleChanged);
+	virtual FOnTitleChanged& OnTitleChanged() override
+	{
+		return TitleChangedEvent;
+	}
+
+	DECLARE_DERIVED_EVENT(FWebBrowserWindow, IWebBrowserWindow::FOnUrlChanged, FOnUrlChanged);
+	virtual FOnUrlChanged& OnUrlChanged() override
+	{
+		return UrlChangedEvent;
+	}
+
+	DECLARE_DERIVED_EVENT(FWebBrowserWindow, IWebBrowserWindow::FOnToolTip, FOnToolTip);
+	virtual FOnToolTip& OnToolTip() override
+	{
+		return ToolTipEvent;
+	}
+
+	DECLARE_DERIVED_EVENT(FWebBrowserWindow, IWebBrowserWindow::FOnNeedsRedraw, FOnNeedsRedraw);
+	virtual FOnNeedsRedraw& OnNeedsRedraw() override
+	{
+		return NeedsRedrawEvent;
+	}
+
+	virtual FOnBeforeBrowse& OnBeforeBrowse() override
+	{
+		return BeforeBrowseDelegate;
+	}
+
+	virtual FOnLoadUrl& OnLoadUrl() override
+	{
+		return LoadUrlDelegate;
+	}
+
+	virtual FOnCreateWindow& OnCreateWindow() override
+	{
+		return CreateWindowDelegate;
+	}
+
+	virtual FOnCloseWindow& OnCloseWindow() override
+	{
+		return CloseWindowDelegate;
+	}
+
+	virtual FOnFloatingCloseButtonPressed& OnFloatingCloseButtonPressed() override
+	{
+		return FloatingCloseButtonPressedDelegate;
+	}
+
+	virtual FCursorReply OnCursorQuery(const FGeometry& MyGeometry, const FPointerEvent& CursorEvent) override
+	{
+		//return Cursor == EMouseCursor::Default ? FCursorReply::Unhandled() : FCursorReply::Cursor(Cursor);
+		return FCursorReply::Unhandled();
+	}
+
+	virtual FOnBeforeResourceLoadDelegate& OnBeforeResourceLoad() override
+	{
+		return BeforeResourceLoadDelegate;
+	}
+
+	virtual FOnResourceLoadCompleteDelegate& OnResourceLoadComplete() override
+	{
+		return ResourceLoadCompleteDelegate;
+	}
+
+	virtual FOnConsoleMessageDelegate& OnConsoleMessage() override
+	{
+		return ConsoleMessageDelegate;
+	}
+
+	virtual FOnBeforePopupDelegate& OnBeforePopup() override
+	{
+		return BeforePopupDelegate;
+	}
+
+	DECLARE_DERIVED_EVENT(FWebBrowserWindow, IWebBrowserWindow::FOnShowPopup, FOnShowPopup);
+	virtual FOnShowPopup& OnShowPopup() override
+	{
+		return ShowPopupEvent;
+	}
+
+	DECLARE_DERIVED_EVENT(FWebBrowserWindow, IWebBrowserWindow::FOnDismissPopup, FOnDismissPopup);
+	virtual FOnDismissPopup& OnDismissPopup() override
+	{
+		return DismissPopupEvent;
+	}
+
+	virtual FOnShowDialog& OnShowDialog() override
+	{
+		return ShowDialogDelegate;
+	}
+
+	virtual FOnDismissAllDialogs& OnDismissAllDialogs() override
+	{
+		return DismissAllDialogsDelegate;
+	}
+
+	virtual FOnSuppressContextMenu& OnSuppressContextMenu() override
+	{
+		return SuppressContextMenuDelgate;
+	}
+
+	virtual FOnDragWindow& OnDragWindow() override
+	{
+		return DragWindowDelegate;
+	}
+	
+	void NotifyDocumentLoadingStateChange(const FString& InCurrentUrl, bool IsLoading);
+
+	void NotifyDocumentError(const FString& InCurrentUrl, int InErrorCode);
+
+	bool OnJsMessageReceived(const FString& Command, const TArray<FString>& Params, const FString& Origin);
+
+	virtual FOnUnhandledKeyDown& OnUnhandledKeyDown() override
+	{
+		return UnhandledKeyDownDelegate;
+	}
+
+	virtual FOnUnhandledKeyUp& OnUnhandledKeyUp() override
+	{
+		return UnhandledKeyUpDelegate;
+	}
+
+	virtual FOnUnhandledKeyChar& OnUnhandledKeyChar() override
+	{
+		return UnhandledKeyCharDelegate;
+	}
+
+	void NotifyUrlChanged(const FString& InCurrentUrl);
+public:
+	/**
+	* Called from the WebBrowserSingleton tick event. Should test whether the widget got a tick from Slate last frame and set the state to hidden if not.
+	*/
+	void CheckTickActivity() override;
+
+	/**
+	* Signal from the widget, meaning that the widget is still active
+	*/
+	void SetTickLastFrame();
+
+	/**
+	* Browser's visibility
+	*/
+	bool IsVisible();
+
+	/**
+	* Webvew floating close button was pressed
+	*/
+	void FloatingCloseButtonPressed();
+
+
+	void SetTitle(const FString& InTitle)
+	{
+		Title = InTitle;
+		OnTitleChanged().Broadcast(Title);
+	}
+
+private:
+
+	TSharedPtr<SAppleWebBrowserWidget> BrowserWidget;
+
+	/** Current title of this window. */
+	FString Title;
+
+	/** Current Url of this window. */
+	FString CurrentUrl;
+
+	/** User-Agent Application to report. */
+	FString UserAgentApplication;
+
+	/** Optional text to load as a web page. */
+	TOptional<FString> ContentsToLoad;
+	
+	/** Whether to enable background transparency */
+	bool bUseTransparency;
+
+	// TODO: None of these events are actually called
+
+	/** Delegate for broadcasting load state changes. */
+	FOnDocumentStateChanged DocumentStateChangedEvent;
+
+	/** Delegate for broadcasting title changes. */
+	FOnTitleChanged TitleChangedEvent;
+
+	/** Delegate for broadcasting address changes. */
+	FOnUrlChanged UrlChangedEvent;
+
+	/** Delegate for broadcasting when the browser wants to show a tool tip. */
+	FOnToolTip ToolTipEvent;
+
+	/** Delegate for notifying that the window needs refreshing. */
+	FOnNeedsRedraw NeedsRedrawEvent;
+
+	/** Delegate that is executed prior to browser navigation. */
+	FOnBeforeBrowse BeforeBrowseDelegate;
+
+	/** Delegate for overriding Url contents. */
+	FOnLoadUrl LoadUrlDelegate;
+
+	/** Delegate for notifying that a popup window is attempting to open. */
+	FOnBeforePopupDelegate BeforePopupDelegate;
+
+	/** Delegate for notifying that the browser is about to load a resource. */
+	FOnBeforeResourceLoadDelegate BeforeResourceLoadDelegate;
+
+	/** Delegate that allows for responses to resource loads */
+	FOnResourceLoadCompleteDelegate ResourceLoadCompleteDelegate;
+
+	/** Delegate that allows for response to console logs.  Typically used to capture and mirror web logs in client application logs. */
+	FOnConsoleMessageDelegate ConsoleMessageDelegate;
+
+	/** Delegate for handling requests to create new windows. */
+	FOnCreateWindow CreateWindowDelegate;
+
+	/** Delegate for handling requests to close new windows that were created. */
+	FOnCloseWindow CloseWindowDelegate;
+
+	/** Delegate for handling requests to close from the webview floating close button. */
+	FOnFloatingCloseButtonPressed FloatingCloseButtonPressedDelegate;
+
+	/** Delegate for handling requests to show the popup menu. */
+	FOnShowPopup ShowPopupEvent;
+
+	/** Delegate for handling requests to dismiss the current popup menu. */
+	FOnDismissPopup DismissPopupEvent;
+
+	/** Delegate for showing dialogs. */
+	FOnShowDialog ShowDialogDelegate;
+
+	/** Delegate for dismissing all dialogs. */
+	FOnDismissAllDialogs DismissAllDialogsDelegate;
+
+	/** Delegate for suppressing context menu */
+	FOnSuppressContextMenu SuppressContextMenuDelgate;
+
+	/** Delegate that is executed when a drag event is detected in an area of the web page tagged as a drag region. */
+	FOnDragWindow DragWindowDelegate;
+
+	/** Current state of the document being loaded. */
+	EWebBrowserDocumentState DocumentState;
+	int ErrorCode;
+
+	FMobileJSScriptingPtr Scripting;
+
+	mutable TOptional<TFunction<void(const FString&)>> GetPageSourceCallback;
+	/** Delegate for handling key down events not handled by the browser. */
+	FOnUnhandledKeyDown UnhandledKeyDownDelegate;
+	
+	/** Delegate for handling key up events not handled by the browser. */
+	FOnUnhandledKeyUp UnhandledKeyUpDelegate;
+	
+	/** Delegate for handling key char events not handled by the browser. */
+	FOnUnhandledKeyChar UnhandledKeyCharDelegate;
+
+	TSharedPtr<SWindow> ParentWindow;
+
+	FIntPoint WindowSize;
+
+	/** Tracks whether the widget is currently disabled or not*/
+	bool bIsDisabled;
+
+	/** Tracks whether the widget is currently visible or not*/
+	bool bIsVisible;
+
+	/** Used to detect when the widget is hidden*/
+	bool bTickedLastFrame;
+};
+
+#endif
diff --git a/Engine/Source/Runtime/WebBrowser/Private/Apple/WebViewCloseButton.cpp b/Engine/Source/Runtime/WebBrowser/Private/Apple/WebViewCloseButton.cpp
new file mode 100644
--- /dev/null
+++ b/Engine/Source/Runtime/WebBrowser/Private/Apple/WebViewCloseButton.cpp
@@ -0,0 +1,400 @@
+// Copyright Epic Games, Inc. All Rights Reserved.
+#include "WebViewCloseButton.h"
+
+#if PLATFORM_IOS
+#if !PLATFORM_TVOS
+
+static constexpr float GFullscreenCloseButtonSize = 40.0f;
+static constexpr float GFullscreenCloseButtonInnerXInnerSize = 32.0f;
+static constexpr float GFullscreenCloseButtonInnerXOuterSize = 34.0f;
+static constexpr float GFullscreenCloseButtonOffsetFromEdge = 4.0f;
+static constexpr float GCloseButtonClickDragTolerance = 10.0f;
+
+static const uint8_t XTextureData[] =
+{
+		0x00, 0x00, 0x00, 0x00, 0x00, 0x00, 0x00, 0x00, 0x00, 0x00, 0x00, 0x00, 0x00, 0x00, 0x00, 0x00,
+		0x00, 0x00, 0x00, 0x00, 0x00, 0x00, 0x00, 0x00, 0x00, 0x00, 0x00, 0x00, 0x00, 0x00, 0x00, 0x00,
+		0x00, 0x00, 0x00, 0x00, 0x00, 0x00, 0x00, 0x00, 0x00, 0x00, 0x00, 0x00, 0x00, 0x00, 0x00, 0x00,
+		0x00, 0x00, 0x00, 0x00, 0x00, 0x00, 0x00, 0x00, 0x00, 0x00, 0x00, 0x00, 0x00, 0x00, 0x00, 0x00,
+		0x00, 0x00, 0x00, 0x00, 0x00, 0x00, 0x00, 0x00, 0x00, 0x00, 0x00, 0x00, 0x00, 0x00, 0x00, 0x00,
+		0x00, 0x00, 0x00, 0x00, 0x00, 0x00, 0x00, 0x00, 0x00, 0x00, 0x00, 0x00, 0x00, 0x00, 0x00, 0x00,
+		0x00, 0x00, 0x00, 0x00, 0x00, 0x00, 0x00, 0x71, 0x93, 0x3d, 0x00, 0x00, 0x00, 0x00, 0x00, 0x00,
+		0x00, 0x00, 0x00, 0x00, 0x00, 0x00, 0x50, 0xe9, 0xfa, 0xe3, 0x58, 0x00, 0x00, 0x00, 0x00, 0x00,
+		0x00, 0x00, 0x00, 0x00, 0x00, 0x00, 0x00, 0xbc, 0xf8, 0xfa, 0xe3, 0x58, 0x00, 0x00, 0x00, 0x00,
+		0x00, 0x00, 0x00, 0x00, 0x00, 0x00, 0x00, 0x1e, 0xc2, 0xf8, 0xfa, 0xe3, 0x58, 0x00, 0x00, 0x00,
+		0x00, 0x00, 0x00, 0x00, 0x00, 0x00, 0x00, 0x00, 0x1e, 0xc2, 0xf8, 0xfa, 0xe3, 0x58, 0x00, 0x00,
+		0x00, 0x00, 0x00, 0x00, 0x00, 0x00, 0x00, 0x00, 0x00, 0x1e, 0xc2, 0xf8, 0xfa, 0xe3, 0x58, 0x00,
+		0x00, 0x00, 0x00, 0x00, 0x00, 0x00, 0x00, 0x00, 0x00, 0x00, 0x1e, 0xc2, 0xf8, 0xfa, 0xe3, 0x58,
+		0x00, 0x00, 0x00, 0x00, 0x00, 0x00, 0x00, 0x00, 0x00, 0x00, 0x00, 0x1e, 0xc2, 0xf8, 0xfa, 0xe5,
+		0x00, 0x00, 0x00, 0x00, 0x00, 0x00, 0x00, 0x00, 0x00, 0x00, 0x00, 0x00, 0x1e, 0xc2, 0xf8, 0xfa,
+		0x00, 0x00, 0x00, 0x00, 0x00, 0x00, 0x00, 0x00, 0x00, 0x00, 0x00, 0x00, 0x00, 0x34, 0xdc, 0xfa,
+};
+
+struct FDefaultSymmetricTextureColorFn
+{
+	uint32_t operator()(uint32_t In) { return In; }
+};
+
+static UIImage* CreateSymmetricTexture(const uint8_t Data[16 * 16], TFunction<uint32_t(uint32_t)>&& ColorFn = FDefaultSymmetricTextureColorFn())
+{
+	uint32_t TextureData[32 * 32];
+
+	for (int Row = 0; Row < 16; ++Row)
+	{
+		for (int Col = 0; Col < 16; ++Col)
+		{
+			const uint32_t Gray = Data[Row * 16 + Col];
+			const uint32_t Color = ColorFn(0xFFFFFF | (Gray << 24));
+
+#define Q1 Row*32+Col
+#define Q2 Row*32+(32 - Col - 1)
+#define Q3 (32 - Row - 1)*32+Col
+#define Q4 (32 - Row - 1)*32+(32 - Col - 1)
+			TextureData[Q1] = Color;
+			TextureData[Q2] = Color;
+			TextureData[Q3] = Color;
+			TextureData[Q4] = Color;
+#undef Q1
+#undef Q2
+#undef Q3
+#undef Q4
+		}
+	}
+
+	UIGraphicsBeginImageContext(CGSizeMake(32, 32));
+	CGColorSpaceRef ColorSpace = CGColorSpaceCreateDeviceRGB();
+	CGBitmapInfo BitmapInfo = kCGImageAlphaFirst | kCGBitmapByteOrder32Little;
+	CGDataProviderRef DataProvider = CGDataProviderCreateWithData(NULL, TextureData, sizeof(TextureData), NULL);
+	CGImageRef ImageRef = CGImageCreate(32, 32, 8, 32, 32 * 4, ColorSpace, BitmapInfo, DataProvider, NULL, NO, kCGRenderingIntentDefault);
+	CGColorSpaceRelease(ColorSpace);
+	CGDataProviderRelease(DataProvider);
+
+	CGContextRef G = UIGraphicsGetCurrentContext();
+	CGContextClearRect(G, CGRectMake(0, 0, 32, 32));
+	CGContextDrawImage(G, CGRectMake(0, 0, 32, 32), ImageRef);
+	UIImage* Result = UIGraphicsGetImageFromCurrentImageContext();
+	UIGraphicsEndImageContext();
+
+	return Result;
+}
+
+static UIImage* CreateTintedImage(UIImage* SourceImage, UIColor* TintColor)
+{
+	bool bNeedsTint = false;
+	CGFloat Red, Green, Blue, Alpha;
+	if ([TintColor getRed:&Red green:&Green blue:&Blue alpha:&Alpha])
+	{
+		bNeedsTint = Red < 0.99 || Green < 0.99 || Blue < 0.99 || Alpha < 0.99;
+	}
+	else
+	{
+		bNeedsTint = !CGColorEqualToColor(UIColor.whiteColor.CGColor, TintColor.CGColor);
+	}
+
+	if (!bNeedsTint)
+	{
+		return SourceImage;
+	}
+
+	size_t ImageWidth = (size_t)(SourceImage.size.width * SourceImage.scale);
+	size_t ImageHeight = (size_t)(SourceImage.size.height * SourceImage.scale);
+	if (!ImageHeight || !ImageWidth)
+	{
+		return SourceImage;
+	}
+
+	uint32_t* PixelData = (uint32_t*)calloc(ImageWidth*ImageHeight, 4);
+	CGColorSpaceRef ColorSpace = CGColorSpaceCreateDeviceRGB();
+
+	CGRect Rect = CGRectMake(0, 0, ImageWidth, ImageHeight);
+	CGRect FirstPixelDataEntryRect = CGRectMake(0, ImageHeight-1, 1, 1); // bottom-most scanline is the first row in the array
+	CGContextRef BitmapContext = CGBitmapContextCreate(PixelData, ImageWidth, ImageHeight, 8, ImageWidth*4, ColorSpace, kCGImageAlphaPremultipliedLast | kCGImageByteOrder32Big);
+	CGContextClearRect(BitmapContext, Rect);
+	CGContextSetBlendMode(BitmapContext, kCGBlendModeNormal);
+	CGContextSetFillColorSpace(BitmapContext, CGColorGetColorSpace(TintColor.CGColor));
+	CGContextSetFillColor(BitmapContext, CGColorGetComponents(TintColor.CGColor));
+	CGContextFillRect(BitmapContext, FirstPixelDataEntryRect); // write a single pixel in the tinting color so we can read it out as RGBA
+	CGContextFlush(BitmapContext);
+
+	// PixelData now contains a single entry with our TintColor -- grab it.
+	const uint32_t PremultipliedTintColor = *PixelData;
+	uint32_t TA = (PremultipliedTintColor >> 24) & 0xFF;
+	uint32_t TB = (PremultipliedTintColor >> 16) & 0xFF;
+	uint32_t TG = (PremultipliedTintColor >> 8) & 0xFF;
+	uint32_t TR = PremultipliedTintColor & 0xFF;
+	if (TA)
+	{
+		TB = TB * 255 / TA + 1;
+		TG = TG * 255 / TA + 1;
+		TR = TR * 255 / TA + 1;
+	}
+
+	CGContextClearRect(BitmapContext, FirstPixelDataEntryRect);
+	CGContextDrawImage(BitmapContext, Rect, SourceImage.CGImage);
+	CGContextFlush(BitmapContext);
+
+	// PixelData now contains the SourceImage pixel data with premultiplied alpha.
+	for (size_t Y = 0; Y < ImageHeight; ++Y)
+	{
+		uint32_t* PixelRow = &PixelData[Y*ImageWidth];
+		for (size_t X = 0; X < ImageWidth; ++X)
+		{
+			uint32_t& Pixel = PixelRow[X];
+			uint32_t SA = (Pixel >> 24) & 0xFF, SB = (Pixel >> 16) & 0xFF, SG = (Pixel >> 8) & 0xFF, SR = Pixel & 0xFF;
+			if (SA)
+			{
+				SB = SB * 255 / SA;
+				SG = SG * 255 / SA;
+				SR = SR * 255 / SA;
+
+				uint32_t NewAlpha = SA*TA/255;
+				NewAlpha = NewAlpha + ((255 - NewAlpha)*TA / 255) + 1;
+
+				SB = (SB*TB*NewAlpha + 0x80) >> 16;
+				SG = (SG*TG*NewAlpha + 0x80) >> 16;
+				SR = (SR*TR*NewAlpha + 0x80) >> 16;
+				SA = NewAlpha - 1;
+				Pixel = (SA << 24) | (SB << 16) | (SG << 8) | SR;
+			}
+		}
+	}
+
+	// Get an image of the pixel data
+	CGImageRef TempImage = CGBitmapContextCreateImage(BitmapContext);
+	UIGraphicsBeginImageContext(CGSizeMake(ImageWidth, ImageHeight));
+	CGContextRef CopyImageContext = UIGraphicsGetCurrentContext();
+	CGContextDrawImage(CopyImageContext, Rect, TempImage);
+	UIImage* OutImage = UIGraphicsGetImageFromCurrentImageContext();
+	UIGraphicsEndImageContext();
+
+	// Release resources
+	CGImageRelease(TempImage);
+	CGContextRelease(BitmapContext);
+	CGColorSpaceRelease(ColorSpace);
+
+	return OutImage;
+}
+
+@interface WebViewCloseButtonDrawingLayer : NSObject
+@property (nonatomic, retain) UIImage* Texture;
+@property (nonatomic) CGFloat TextureInset;
+@property (nonatomic) BOOL bDrawWhileTouchDown;
+@property (nonatomic) CGBlendMode BlendMode;
+@end
+
+@implementation WebViewCloseButtonDrawingLayer
+
++ (instancetype)drawingLayerWithTexture:(UIImage*)Texture blendMode:(CGBlendMode)BlendMode inset:(CGFloat)TextureInset forTouchDown:(BOOL)bDrawWhileTouchDown
+{
+	WebViewCloseButtonDrawingLayer* Layer = [self new];
+	Layer.Texture = Texture;
+	Layer.TextureInset = TextureInset;
+	Layer.bDrawWhileTouchDown = bDrawWhileTouchDown;
+	Layer.BlendMode = BlendMode;
+	return Layer;
+}
+
+@end
+
+@implementation WebViewCloseButton
+{
+	NSMutableArray* DrawingLayers;
+	BOOL bIsDraggable;
+	BOOL bIsTouchDown;
+	BOOL bMayBeDragged;
+	CGPoint TouchDownPoint;
+	CGPoint DownDisplacement;
+	CGPoint Displacement;
+	CGPoint LayoutPosition;
+	CGRect SafeBounds;
+}
+
+- (void)dealloc
+{
+	[DrawingLayers release];
+	
+	[super dealloc];
+}
+
+- (instancetype)initWithFrame:(CGRect)frame
+{
+	if (self = [super initWithFrame:frame])
+	{
+		DrawingLayers = [[NSMutableArray arrayWithCapacity:4] retain];
+		self.hidden = YES;
+		bIsTouchDown = NO;
+		bMayBeDragged = NO;
+		bIsDraggable = NO;
+	}
+	return self;
+}
+
+-(void)showButton:(BOOL)bShow setDraggable:(BOOL)bDraggable
+{
+	self.hidden = !bShow;
+	if (bIsDraggable != bDraggable)
+	{
+		bIsDraggable = bDraggable;
+		[self setupLayout];
+	}
+}
+
+- (void)handleTap
+{
+	if (self.TapHandler)
+	{
+		self.TapHandler();
+	}
+}
+
+- (void)addDrawingLayerTexture:(UIImage *)Texture textureMultiply:(nonnull UIColor *)TextureMultiply blendMode:(CGBlendMode)BlendMode inset:(CGFloat)TextureInset forTouchDown:(BOOL)bInIsTouchDown
+{
+	if (Texture)
+	{
+		Texture = CreateTintedImage(Texture, TextureMultiply);
+		[DrawingLayers addObject:[WebViewCloseButtonDrawingLayer drawingLayerWithTexture:Texture blendMode:BlendMode inset:TextureInset forTouchDown:bInIsTouchDown]];
+		[self setNeedsDisplay];
+	}
+}
+
+- (void)setupLayout
+{
+	bMayBeDragged = NO;
+	Displacement = CGPointMake(0,0);
+	self.bounds = CGRectMake(0,0, GFullscreenCloseButtonSize, GFullscreenCloseButtonSize);
+	
+	if (!self.superview)
+	{
+		return;
+	}
+
+	UIEdgeInsets SuperInsets = [self.superview safeAreaInsets];
+	
+	// Substract insets from parent bounds
+	SafeBounds = self.superview.bounds;
+	SafeBounds.origin.x += SuperInsets.left;
+	SafeBounds.origin.y += SuperInsets.top;
+	SafeBounds.size.width -= SuperInsets.left + SuperInsets.right;
+	SafeBounds.size.height -= SuperInsets.top + SuperInsets.bottom;
+	
+	// Substract space for some margins and size of the brutton from its center
+	SafeBounds = CGRectInset(SafeBounds, GFullscreenCloseButtonSize/2 + GFullscreenCloseButtonOffsetFromEdge, GFullscreenCloseButtonSize/2 + GFullscreenCloseButtonOffsetFromEdge);
+
+	// Position the button in the top right corner fo the safe bounds
+	LayoutPosition = CGPointMake(SafeBounds.origin.x + SafeBounds.size.width, SafeBounds.origin.y);
+	self.center = LayoutPosition;
+}
+
+- (void)drawRect:(CGRect)rect
+{
+	for (WebViewCloseButtonDrawingLayer* Layer in DrawingLayers)
+	{
+		if (Layer.bDrawWhileTouchDown != bIsTouchDown)
+		{
+			continue;
+		}
+
+		CGRect RectCopy = CGRectInset(rect, Layer.TextureInset, Layer.TextureInset);
+		[Layer.Texture drawInRect:RectCopy blendMode:Layer.BlendMode alpha:1.0];
+	}
+}
+
+- (void)touchesBegan:(NSSet<UITouch *> *)touches withEvent:(UIEvent *)event
+{
+	if (touches.count > 1)
+	{
+		return;
+	}
+
+	bIsTouchDown = [self pointInside:[touches.anyObject locationInView:self] withEvent:event];
+	if (bIsTouchDown)
+	{
+		bMayBeDragged = bIsDraggable;
+		UITouch* Touch = [[touches allObjects] objectAtIndex: 0];
+		TouchDownPoint = [Touch locationInView: self.superview];
+		DownDisplacement = CGPointMake(TouchDownPoint.x - Displacement.x, TouchDownPoint.y - Displacement.y);
+	}
+	[self setNeedsDisplay];
+}
+
+- (void)touchesMoved:(NSSet<UITouch *> *)touches withEvent:(UIEvent *)event
+{
+	if (touches.count > 1)
+	{
+		return;
+	}
+
+	if (bMayBeDragged)
+	{
+		UITouch* Touch = [[touches allObjects] objectAtIndex: 0];
+		CGPoint Point = [Touch locationInView: self.superview];
+		CGPoint AttemptedDisplacement = CGPointMake(Point.x - DownDisplacement.x, Point.y - DownDisplacement.y);
+
+		CGPoint UpdatedCenter = CGPointMake( FMath::Clamp(LayoutPosition.x + AttemptedDisplacement.x, CGRectGetMinX(SafeBounds), CGRectGetMaxX(SafeBounds)),
+											 FMath::Clamp(LayoutPosition.y + AttemptedDisplacement.y, CGRectGetMinY(SafeBounds), CGRectGetMaxY(SafeBounds)));
+		
+		Displacement = CGPointMake(UpdatedCenter.x - LayoutPosition.x, UpdatedCenter.y - LayoutPosition.y);
+		
+		self.center = UpdatedCenter;
+	}
+	bIsTouchDown = [self pointInside:[touches.anyObject locationInView:self] withEvent:event];
+	[self setNeedsDisplay];
+}
+
+- (void)touchesEnded:(NSSet<UITouch *> *)touches withEvent:(UIEvent *)event
+{
+	if (touches.count > 1)
+	{
+		return;
+	}
+
+	bIsTouchDown = NO;
+	[self setNeedsDisplay];
+
+	UITouch* Touch = [[touches allObjects] objectAtIndex: 0];
+	CGPoint Point = [Touch locationInView: self.superview];
+	CGPoint AbsDisplacement = CGPointMake(Point.x - TouchDownPoint.x, Point.y - TouchDownPoint.y);
+	
+	if (FMath::Abs(AbsDisplacement.x) < GCloseButtonClickDragTolerance && FMath::Abs(AbsDisplacement.y) < GCloseButtonClickDragTolerance )
+	{
+		if ([self pointInside:[touches.anyObject locationInView:self] withEvent:event])
+		{
+			[self handleTap];
+		}
+	}
+}
+
+- (void)touchesCancelled:(NSSet<UITouch *> *)touches withEvent:(UIEvent *)event
+{
+	bIsTouchDown = NO;
+	bMayBeDragged = NO;
+	[self setNeedsDisplay];
+}
+
+- (void)touchesEstimatedPropertiesUpdated:(NSSet<UITouch *> *)touches
+{
+}
+
+@end
+
+WebViewCloseButton* MakeCloseButton()
+{
+	WebViewCloseButton* Button = [[WebViewCloseButton alloc] initWithFrame:CGRectZero];
+	Button.translatesAutoresizingMaskIntoConstraints = NO;
+	Button.opaque = NO;
+
+	UIImage* XImage = CreateSymmetricTexture(XTextureData);
+
+	// X states
+	CGFloat InsetWhite = (GFullscreenCloseButtonSize - GFullscreenCloseButtonInnerXInnerSize) / 2.0;
+	CGFloat InsetDark = (GFullscreenCloseButtonSize - GFullscreenCloseButtonInnerXOuterSize) / 2.0;
+	[Button addDrawingLayerTexture:XImage textureMultiply:UIColor.blackColor blendMode:kCGBlendModeNormal inset:InsetDark forTouchDown:NO];
+	[Button addDrawingLayerTexture:XImage textureMultiply:UIColor.whiteColor blendMode:kCGBlendModeNormal inset:InsetWhite forTouchDown:NO];
+	[Button addDrawingLayerTexture:XImage textureMultiply:UIColor.whiteColor blendMode:kCGBlendModeNormal inset:InsetDark forTouchDown:YES];
+	[Button addDrawingLayerTexture:XImage textureMultiply:UIColor.blackColor blendMode:kCGBlendModeNormal inset:InsetWhite forTouchDown:YES];
+
+	return Button;
+}
+
+#endif // !PLATFORM_TVOS
+#endif // PLATFORM_IOS
diff --git a/Engine/Source/Runtime/WebBrowser/Private/Apple/WebViewCloseButton.h b/Engine/Source/Runtime/WebBrowser/Private/Apple/WebViewCloseButton.h
new file mode 100644
--- /dev/null
+++ b/Engine/Source/Runtime/WebBrowser/Private/Apple/WebViewCloseButton.h
@@ -0,0 +1,23 @@
+// Copyright Epic Games, Inc. All Rights Reserved.
+
+#pragma once
+
+#include "CoreMinimal.h"
+
+#if PLATFORM_IOS
+
+#if !PLATFORM_TVOS
+#import <UIKit/UIKit.h>
+
+@interface WebViewCloseButton : UIView
+@property (nonatomic, copy) void (^TapHandler)(void);
+
+- (void)setupLayout;
+- (void)showButton:(BOOL)bShow setDraggable:(BOOL)bDraggable;
+
+@end
+
+WebViewCloseButton* MakeCloseButton();
+#endif
+
+#endif
diff --git a/Engine/Source/Runtime/WebBrowser/Private/IOS/IOSCookieManager.cpp b/Engine/Source/Runtime/WebBrowser/Private/IOS/IOSCookieManager.cpp
deleted file mode 100644
diff --git a/Engine/Source/Runtime/WebBrowser/Private/IOS/IOSCookieManager.h b/Engine/Source/Runtime/WebBrowser/Private/IOS/IOSCookieManager.h
deleted file mode 100644
diff --git a/Engine/Source/Runtime/WebBrowser/Private/IOS/IOSPlatformWebBrowser.cpp b/Engine/Source/Runtime/WebBrowser/Private/IOS/IOSPlatformWebBrowser.cpp
deleted file mode 100644
diff --git a/Engine/Source/Runtime/WebBrowser/Private/IOS/IOSPlatformWebBrowser.h b/Engine/Source/Runtime/WebBrowser/Private/IOS/IOSPlatformWebBrowser.h
deleted file mode 100644
diff --git a/Engine/Source/Runtime/WebBrowser/Private/IOS/WebViewCloseButton.cpp b/Engine/Source/Runtime/WebBrowser/Private/IOS/WebViewCloseButton.cpp
deleted file mode 100644
diff --git a/Engine/Source/Runtime/WebBrowser/Private/IOS/WebViewCloseButton.h b/Engine/Source/Runtime/WebBrowser/Private/IOS/WebViewCloseButton.h
deleted file mode 100644
diff --git a/Engine/Source/Runtime/WebBrowser/Private/MobileJS/MobileJSScripting.cpp b/Engine/Source/Runtime/WebBrowser/Private/MobileJS/MobileJSScripting.cpp
--- a/Engine/Source/Runtime/WebBrowser/Private/MobileJS/MobileJSScripting.cpp
+++ b/Engine/Source/Runtime/WebBrowser/Private/MobileJS/MobileJSScripting.cpp
@@ -2,7 +2,7 @@
 
 #include "MobileJSScripting.h"
 
-#if PLATFORM_ANDROID  || PLATFORM_IOS
+#if PLATFORM_ANDROID  || PLATFORM_IOS  || PLATFORM_MAC
 
 #include "IWebBrowserWindow.h"
 #include "MobileJSStructSerializerBackend.h"
@@ -664,4 +664,4 @@ void FMobileJSScripting::SetWindow(const TSharedRef<class IWebBrowserWindow>& In
 	WindowPtr = InWindow;
 }
 
-#endif // PLATFORM_ANDROID  || PLATFORM_IOS
+#endif // PLATFORM_ANDROID  || PLATFORM_IOS  || PLATFORM_MAC
diff --git a/Engine/Source/Runtime/WebBrowser/Private/MobileJS/MobileJSScripting.h b/Engine/Source/Runtime/WebBrowser/Private/MobileJS/MobileJSScripting.h
--- a/Engine/Source/Runtime/WebBrowser/Private/MobileJS/MobileJSScripting.h
+++ b/Engine/Source/Runtime/WebBrowser/Private/MobileJS/MobileJSScripting.h
@@ -3,7 +3,7 @@
 
 #include "CoreMinimal.h"
 
-#if PLATFORM_ANDROID  || PLATFORM_IOS
+#if PLATFORM_ANDROID  || PLATFORM_IOS  || PLATFORM_MAC
 
 #include "WebJSFunction.h"
 #include "WebJSScripting.h"
@@ -69,4 +69,4 @@ private:
 	bool bInjectJSOnPageStarted;
 };
 
-#endif // PLATFORM_ANDROID  || PLATFORM_IOS
+#endif // PLATFORM_ANDROID  || PLATFORM_IOS  || PLATFORM_MAC
diff --git a/Engine/Source/Runtime/WebBrowser/Private/MobileJS/MobileJSStructDeserializerBackend.cpp b/Engine/Source/Runtime/WebBrowser/Private/MobileJS/MobileJSStructDeserializerBackend.cpp
--- a/Engine/Source/Runtime/WebBrowser/Private/MobileJS/MobileJSStructDeserializerBackend.cpp
+++ b/Engine/Source/Runtime/WebBrowser/Private/MobileJS/MobileJSStructDeserializerBackend.cpp
@@ -2,7 +2,7 @@
 
 #include "MobileJSStructDeserializerBackend.h"
 
-#if	PLATFORM_ANDROID || PLATFORM_IOS
+#if	PLATFORM_ANDROID || PLATFORM_IOS || PLATFORM_MAC
 
 #include "MobileJSScripting.h"
 #include "UObject/UnrealType.h"
@@ -105,4 +105,4 @@ FMobileJSStructDeserializerBackend::FMobileJSStructDeserializerBackend(FMobileJS
 	JsonData.Append((uint8*)UTF16String.Get(), UTF16String.Length() * sizeof(UTF16CHAR));
 }
 
-#endif // PLATFORM_ANDROID || PLATFORM_IOS
\ No newline at end of file
+#endif // PLATFORM_ANDROID || PLATFORM_IOS || PLATFORM_MAC
diff --git a/Engine/Source/Runtime/WebBrowser/Private/MobileJS/MobileJSStructDeserializerBackend.h b/Engine/Source/Runtime/WebBrowser/Private/MobileJS/MobileJSStructDeserializerBackend.h
--- a/Engine/Source/Runtime/WebBrowser/Private/MobileJS/MobileJSStructDeserializerBackend.h
+++ b/Engine/Source/Runtime/WebBrowser/Private/MobileJS/MobileJSStructDeserializerBackend.h
@@ -4,7 +4,7 @@
 
 #include "CoreMinimal.h"
 
-#if PLATFORM_ANDROID  || PLATFORM_IOS
+#if PLATFORM_ANDROID  || PLATFORM_IOS  || PLATFORM_MAC
 
 #include "MobileJSScripting.h"
 #include "Backends/JsonStructDeserializerBackend.h"
@@ -24,4 +24,4 @@ private:
 	FMemoryReader Reader;
 };
 
-#endif // USE_ANDROID_JNI || PLATFORM_IOS
\ No newline at end of file
+#endif // USE_ANDROID_JNI || PLATFORM_IOS
diff --git a/Engine/Source/Runtime/WebBrowser/Private/MobileJS/MobileJSStructSerializerBackend.cpp b/Engine/Source/Runtime/WebBrowser/Private/MobileJS/MobileJSStructSerializerBackend.cpp
--- a/Engine/Source/Runtime/WebBrowser/Private/MobileJS/MobileJSStructSerializerBackend.cpp
+++ b/Engine/Source/Runtime/WebBrowser/Private/MobileJS/MobileJSStructSerializerBackend.cpp
@@ -2,7 +2,7 @@
 
 #include "MobileJSStructSerializerBackend.h"
 
-#if PLATFORM_ANDROID || PLATFORM_IOS
+#if PLATFORM_ANDROID || PLATFORM_IOS || PLATFORM_MAC
 
 #include "MobileJSScripting.h"
 #include "UObject/UnrealType.h"
@@ -58,4 +58,4 @@ FMobileJSStructSerializerBackend::FMobileJSStructSerializerBackend(TSharedRef<cl
 {
 }
 
-#endif // PLATFORM_ANDROID || PLATFORM_IOS
\ No newline at end of file
+#endif // PLATFORM_ANDROID || PLATFORM_IOS || PLATFORM_MAC
diff --git a/Engine/Source/Runtime/WebBrowser/Private/MobileJS/MobileJSStructSerializerBackend.h b/Engine/Source/Runtime/WebBrowser/Private/MobileJS/MobileJSStructSerializerBackend.h
--- a/Engine/Source/Runtime/WebBrowser/Private/MobileJS/MobileJSStructSerializerBackend.h
+++ b/Engine/Source/Runtime/WebBrowser/Private/MobileJS/MobileJSStructSerializerBackend.h
@@ -4,7 +4,7 @@
 
 #include "CoreMinimal.h"
 
-#if PLATFORM_ANDROID || PLATFORM_IOS
+#if PLATFORM_ANDROID || PLATFORM_IOS || PLATFORM_MAC
 
 #include "MobileJSScripting.h"
 #include "Backends/JsonStructSerializerBackend.h"
@@ -42,4 +42,4 @@ private:
 	FMemoryWriter Writer;
 };
 
-#endif // PLATFORM_ANDROID || PLATFORM_IOS
\ No newline at end of file
+#endif // PLATFORM_ANDROID || PLATFORM_IOS || PLATFORM_MAC
diff --git a/Engine/Source/Runtime/WebBrowser/Private/SWebBrowserView.cpp b/Engine/Source/Runtime/WebBrowser/Private/SWebBrowserView.cpp
--- a/Engine/Source/Runtime/WebBrowser/Private/SWebBrowserView.cpp
+++ b/Engine/Source/Runtime/WebBrowser/Private/SWebBrowserView.cpp
@@ -15,12 +15,12 @@
 
 #if PLATFORM_ANDROID && USE_ANDROID_JNI
 #	include "Android/AndroidWebBrowserWindow.h"
-#elif PLATFORM_IOS
-#	include "IOS/IOSPlatformWebBrowser.h"
-#elif PLATFORM_SPECIFIC_WEB_BROWSER
-#	include COMPILED_PLATFORM_HEADER(PlatformWebBrowser.h)
 #elif WITH_CEF3
 #	include "CEF/CEFWebBrowserWindow.h"
+#elif PLATFORM_IOS || PLATFORM_MAC
+#	include "Apple/ApplePlatformWebBrowser.h"
+#elif PLATFORM_SPECIFIC_WEB_BROWSER
+#	include COMPILED_PLATFORM_HEADER(PlatformWebBrowser.h)
 #else
 #	define DUMMY_WEB_BROWSER 1
 #endif
diff --git a/Engine/Source/Runtime/WebBrowser/Private/WebBrowserSingleton.cpp b/Engine/Source/Runtime/WebBrowser/Private/WebBrowserSingleton.cpp
--- a/Engine/Source/Runtime/WebBrowser/Private/WebBrowserSingleton.cpp
+++ b/Engine/Source/Runtime/WebBrowser/Private/WebBrowserSingleton.cpp
@@ -51,9 +51,9 @@ THIRD_PARTY_INCLUDES_END
 #if PLATFORM_ANDROID && USE_ANDROID_JNI
 #	include "Android/AndroidWebBrowserWindow.h"
 #	include <Android/AndroidCookieManager.h>
-#elif PLATFORM_IOS
-#	include <IOS/IOSPlatformWebBrowser.h>
-#	include <IOS/IOSCookieManager.h>
+#elif PLATFORM_IOS || PLATFORM_MAC
+#	include <Apple/ApplePlatformWebBrowser.h>
+#	include <Apple/AppleCookieManager.h>
 #elif PLATFORM_SPECIFIC_WEB_BROWSER
 #	include COMPILED_PLATFORM_HEADER(PlatformWebBrowser.h)
 #endif
@@ -237,7 +237,7 @@ FWebBrowserSingleton::FWebBrowserSingleton(const FWebBrowserInitSettings& WebBro
 	: WebBrowserWindowFactory(MakeShareable(new FWebBrowserWindowFactory()))
 #else
 	: WebBrowserWindowFactory(MakeShareable(new FNoWebBrowserWindowFactory()))
-#if PLATFORM_IOS || (PLATFORM_ANDROID && USE_ANDROID_JNI)
+#if PLATFORM_IOS || PLATFORM_MAC || (PLATFORM_ANDROID && USE_ANDROID_JNI)
 	, UserAgentApplication(WebBrowserInitSettings.ProductVersion)
 #endif
 #endif
@@ -383,8 +383,8 @@ FWebBrowserSingleton::FWebBrowserSingleton(const FWebBrowserInitSettings& WebBro
 
 		DefaultCookieManager = FCefWebBrowserCookieManagerFactory::Create(CefCookieManager::GetGlobalManager(nullptr));
 	}
-#elif PLATFORM_IOS && !BUILD_EMBEDDED_APP
-	DefaultCookieManager = MakeShareable(new FIOSCookieManager());
+#elif (PLATFORM_MAC || PLATFORM_IOS) && !BUILD_EMBEDDED_APP
+	DefaultCookieManager = MakeShareable(new FAppleCookieManager());
 #elif PLATFORM_ANDROID
 	DefaultCookieManager = MakeShareable(new FAndroidCookieManager());
 #endif
@@ -665,7 +665,7 @@ TSharedPtr<IWebBrowserWindow> FWebBrowserSingleton::CreateBrowserWindow(const FC
 		WindowInterfaces.Add(NewBrowserWindow);
 	}
 	return NewBrowserWindow;
-#elif PLATFORM_IOS
+#elif PLATFORM_IOS || PLATFORM_MAC
 	// Create new window
 	TSharedPtr<FWebBrowserWindow> NewBrowserWindow = MakeShareable(new FWebBrowserWindow(
 		WindowSettings.InitialURL, 
diff --git a/Engine/Source/Runtime/WebBrowser/Private/WebBrowserSingleton.h b/Engine/Source/Runtime/WebBrowser/Private/WebBrowserSingleton.h
--- a/Engine/Source/Runtime/WebBrowser/Private/WebBrowserSingleton.h
+++ b/Engine/Source/Runtime/WebBrowser/Private/WebBrowserSingleton.h
@@ -175,7 +175,7 @@ private:
 	/** List of currently existing browser windows */
 #if WITH_CEF3
 	TArray<TWeakPtr<FCEFWebBrowserWindow>> WindowInterfaces;
-#elif PLATFORM_IOS || PLATFORM_SPECIFIC_WEB_BROWSER || (PLATFORM_ANDROID && USE_ANDROID_JNI)
+#elif PLATFORM_IOS || PLATFORM_MAC || PLATFORM_SPECIFIC_WEB_BROWSER || (PLATFORM_ANDROID && USE_ANDROID_JNI)
 	TArray<TWeakPtr<IWebBrowserWindow>>	WindowInterfaces;
 #endif
 
@@ -184,7 +184,7 @@ private:
 
 	TSharedRef<IWebBrowserWindowFactory> WebBrowserWindowFactory;
 
-#if !WITH_CEF3 && (PLATFORM_IOS || (PLATFORM_ANDROID && USE_ANDROID_JNI))
+#if !WITH_CEF3 && (PLATFORM_IOS || PLATFORM_MAC || (PLATFORM_ANDROID && USE_ANDROID_JNI))
 	FString UserAgentApplication;
 #endif
 
diff --git a/Engine/Source/Runtime/WebBrowser/WebBrowser.Build.cs b/Engine/Source/Runtime/WebBrowser/WebBrowser.Build.cs
--- a/Engine/Source/Runtime/WebBrowser/WebBrowser.Build.cs
+++ b/Engine/Source/Runtime/WebBrowser/WebBrowser.Build.cs
@@ -29,19 +29,31 @@ public class WebBrowser : ModuleRules
         );
 
         if (Target.Platform == UnrealTargetPlatform.Android ||
-		    Target.IsInPlatformGroup(UnrealPlatformGroup.IOS))
+		    Target.IsInPlatformGroup(UnrealPlatformGroup.IOS) || 
+		    Target.Platform == UnrealTargetPlatform.Mac)
 		{
 			// We need these on mobile for external texture support
 			PrivateDependencyModuleNames.AddRange(
 				new string[]
 				{
 					"Engine",
-					"Launch",
 					"JsonUtilities",
-					"WebBrowserTexture"
+					"WebBrowserTexture",
+					"Json"
 				}
 			);
 
+			if(Target.Platform != UnrealTargetPlatform.Mac) 
+			{
+				PrivateDependencyModuleNames.Add("Launch");
+			} 
+			else 
+			{
+				PrivateDependencyModuleNames.Add("RenderCore");
+
+				PublicFrameworks.Add("WebKit");
+			}
+
 			// We need this one on Android for URL decoding
 			PrivateDependencyModuleNames.Add("HTTP");
 
diff --git a/Engine/Source/Runtime/WebBrowserTexture/WebBrowserTexture.Build.cs b/Engine/Source/Runtime/WebBrowserTexture/WebBrowserTexture.Build.cs
--- a/Engine/Source/Runtime/WebBrowserTexture/WebBrowserTexture.Build.cs
+++ b/Engine/Source/Runtime/WebBrowserTexture/WebBrowserTexture.Build.cs
@@ -11,6 +11,7 @@ public class WebBrowserTexture : ModuleRules
         // cook the asset so we must include it in editor builds
         if (Target.Platform == UnrealTargetPlatform.Android ||
             Target.IsInPlatformGroup(UnrealPlatformGroup.IOS) ||
+            Target.Platform == UnrealTargetPlatform.Mac ||
             Target.bBuildEditor == true)
 		{			
 			// Needed for external texture support
