WebBrowser: Fix crash if a FCEFBrowserHandler::GetRootScreenRect gets called during shutdown

This can happen if a specific message is received from CEF and gets handled during the shutdown phase, after the FSlateApplication singleton has already been cleaned up.

#rnx

diff --git a/Engine/Source/Runtime/WebBrowser/Private/CEF/CEFBrowserHandler.cpp b/Engine/Source/Runtime/WebBrowser/Private/CEF/CEFBrowserHandler.cpp
--- a/Engine/Source/Runtime/WebBrowser/Private/CEF/CEFBrowserHandler.cpp
+++ b/Engine/Source/Runtime/WebBrowser/Private/CEF/CEFBrowserHandler.cpp
@@ -298,7 +298,7 @@ void FCEFBrowserHandler::OnLoadingStateChange(CefRefPtr<CefBrowser> Browser, boo
 
 bool FCEFBrowserHandler::GetRootScreenRect(CefRefPtr<CefBrowser> Browser, CefRect& Rect)
 {
-	if (CefCurrentlyOn(TID_UI))
+	if (CefCurrentlyOn(TID_UI) && FSlateApplication::IsInitialized())
 	{
 		// CEF may call this off the main gamethread which slate requires, so double check here
 		FDisplayMetrics DisplayMetrics;
