Getting Engine Window Tabbing Working with WKWebView Web Browser plugin backend on macOS

diff --git a/Engine/Source/Runtime/WebBrowser/Private/Apple/ApplePlatformWebBrowser.cpp b/Engine/Source/Runtime/WebBrowser/Private/Apple/ApplePlatformWebBrowser.cpp
--- a/Engine/Source/Runtime/WebBrowser/Private/Apple/ApplePlatformWebBrowser.cpp
+++ b/Engine/Source/Runtime/WebBrowser/Private/Apple/ApplePlatformWebBrowser.cpp
@@ -144,6 +144,13 @@ class SAppleWebBrowserWidget : public SLeafWidget
 			
 			// CGRect and NSRect have the same memory layout
 			NewFrame = NSRectToCGRect(Rect);
+			
+			if(PastEnclosedWindow != EnclosedWindow) 
+			{
+				WindowFrame = NewFrame;
+				PastEnclosedWindow = EnclosedWindow;
+				LeaveDeadWindow();
+			}
 		}
 #else
 		{
@@ -163,7 +170,8 @@ class SAppleWebBrowserWidget : public SLeafWidget
 		}
 #endif
 
-		if(NewFrame.size.width != WindowFrame.size.width || NewFrame.size.height != WindowFrame.size.height) {
+		if(NewFrame.size.width != WindowFrame.size.width || NewFrame.size.height != WindowFrame.size.height) 
+		{
 			WindowFrame = NewFrame;
 		}
 		
@@ -557,6 +565,14 @@ class SAppleWebBrowserWidget : public SLeafWidget
 		}
 	}
 	
+	void LeaveDeadWindow() 
+	{
+		if(WebViewWrapper != nil) 
+		{
+			[WebViewWrapper leaveDeadWindow];
+		}
+	}
+	
 	~SAppleWebBrowserWidget()
 	{
 		Close();
@@ -565,6 +581,7 @@ class SAppleWebBrowserWidget : public SLeafWidget
 protected:
 	mutable __strong AppleWebViewWrapper* WebViewWrapper;
 private:
+	TWeakPtr<FGenericWindow> PastEnclosedWindow;
 	TWeakPtr<FWebBrowserWindow> WebBrowserWindowPtr;
 
 	/** Enable 3D appearance */
@@ -768,6 +785,14 @@ supportsMetal : (bool)InSupportsMetal supportsMetalMRT : (bool)InSupportsMetalMR
 	[super dealloc];
 }
 
+-(void)leaveDeadWindow;
+{
+	if(!bNeedsAddToView && !Is3DBrowser) 
+	{
+		bNeedsAddToView = true;
+	}
+}
+
 -(void)updateframe:(CGRect)InFrame window:(FGenericWindow*)EnclosedWindow;
 {
 	self.DesiredFrame = InFrame;
@@ -783,7 +808,7 @@ supportsMetal : (bool)InSupportsMetal supportsMetalMRT : (bool)InSupportsMetalMR
 			{
 				bNeedsAddToView = false;
 				
-				#if PLATFORM_IOS 
+#if PLATFORM_IOS 
 				[self.WebViewContainer setOpaque : NO];
 				[self.WebViewContainer setBackgroundColor : [UIColor clearColor]];
 				[[IOSAppDelegate GetDelegate].IOSView addSubview : WebViewContainer];
@@ -792,13 +817,12 @@ supportsMetal : (bool)InSupportsMetal supportsMetalMRT : (bool)InSupportsMetalMR
 				{
 					[CloseButton setupLayout];
 				}
-				#else
-				// Are we adding our WebViewContainer to the proper Window here?
+#else				
 				FCocoaWindow *CocoaWindow = ((FMacWindow*)EnclosedWindow)->GetWindowHandle();
 				NSView *CocoaWindowView = CocoaWindow.openGLView;
 				
-				[CocoaWindowView addSubview : WebViewContainer];
-				#endif
+				[CocoaWindowView addSubview : self.WebViewContainer];
+#endif
 			}
 			else
 			{
@@ -809,14 +833,13 @@ supportsMetal : (bool)InSupportsMetal supportsMetalMRT : (bool)InSupportsMetalMR
 					NextContent = nil;
 					NextURL = nil;
 				}
-				else
-					if (NextURL != nil)
-					{
-						// Load web content from URL
-						NSURLRequest *nsrequest = [NSURLRequest requestWithURL : NextURL];
-						[self.WebView loadRequest : nsrequest];
-						NextURL = nil;
-					}
+				else if (NextURL != nil)
+				{
+					// Load web content from URL
+					NSURLRequest *nsrequest = [NSURLRequest requestWithURL : NextURL];
+					[self.WebView loadRequest : nsrequest];
+					NextURL = nil;
+				}
 			}
 		}
 	});
diff --git a/Engine/Source/Runtime/WebBrowser/Private/Apple/ApplePlatformWebBrowser.h b/Engine/Source/Runtime/WebBrowser/Private/Apple/ApplePlatformWebBrowser.h
--- a/Engine/Source/Runtime/WebBrowser/Private/Apple/ApplePlatformWebBrowser.h
+++ b/Engine/Source/Runtime/WebBrowser/Private/Apple/ApplePlatformWebBrowser.h
@@ -89,6 +89,7 @@ supportsMetal : (bool)InSupportsMetal supportsMetalMRT : (bool)InSupportsMetalMR
 -(bool)IsVideoTextureValid;
 -(bool)UpdateVideoFrame:(void*)ptr;
 -(void)updateWebViewMetalTexture : (id<MTLTexture>)texture;
+-(void)leaveDeadWindow;
 #if !PLATFORM_TVOS
 -(void)webView:(WKWebView*)InWebView decidePolicyForNavigationAction : (WKNavigationAction*)InNavigationAction decisionHandler : (void(^)(WKNavigationActionPolicy))InDecisionHandler;
 -(void)webView:(WKWebView*)InWebView didCommitNavigation : (WKNavigation*)InNavigation;
diff --git a/Engine/Source/Runtime/WebBrowser/Private/SWebBrowserView.cpp b/Engine/Source/Runtime/WebBrowser/Private/SWebBrowserView.cpp
--- a/Engine/Source/Runtime/WebBrowser/Private/SWebBrowserView.cpp
+++ b/Engine/Source/Runtime/WebBrowser/Private/SWebBrowserView.cpp
@@ -250,11 +250,12 @@ int32 SWebBrowserView::OnPaint(const FPaintArgs& Args, const FGeometry& Allotted
 	int32 Layer = SCompoundWidget::OnPaint(Args, AllottedGeometry, MyCullingRect, OutDrawElements, LayerId, InWidgetStyle, bParentEnabled);
 	
 	// Cache a reference to our parent window, if we didn't already reference it.
-	if (!SlateParentWindowPtr.IsValid())
-	{
-		SWindow* ParentWindow = OutDrawElements.GetPaintWindow();
+	SWindow* ParentWindow = OutDrawElements.GetPaintWindow();
 
-		TSharedRef<SWindow> SlateParentWindowRef = StaticCastSharedRef<SWindow>(ParentWindow->AsShared());
+	TSharedRef<SWindow> SlateParentWindowRef = StaticCastSharedRef<SWindow>(ParentWindow->AsShared());
+	
+	if (!SlateParentWindowPtr.IsValid() || SlateParentWindowPtr != SlateParentWindowRef)
+	{
 		SlateParentWindowPtr = SlateParentWindowRef;
 		if (BrowserWindow.IsValid())
 		{
