WebBrowser: Remove obsolete filtering of newlines in injected content

This was a leftover from earlier code which passed the content via a custom header, where newlines are forbidden. It was added in CL 3369255.
Around the same time, the custom header was replaced by POST data, which has no newline restrictions, in CL 3258614.
Then both changes got merged in parallel, and the filtering was left in even though it's not needed anymore.
We can safely get rid of it, especially as it can interfere with injected JS scripts.

From Github PR #13052

#rnx

diff --git a/Engine/Source/Runtime/WebBrowser/Private/CEF/CEFBrowserHandler.cpp b/Engine/Source/Runtime/WebBrowser/Private/CEF/CEFBrowserHandler.cpp
--- a/Engine/Source/Runtime/WebBrowser/Private/CEF/CEFBrowserHandler.cpp
+++ b/Engine/Source/Runtime/WebBrowser/Private/CEF/CEFBrowserHandler.cpp
@@ -484,9 +484,6 @@ CefResourceRequestHandler::ReturnValue FCEFBrowserHandler::OnBeforeResourceLoad(
 			TOptional<FString> Contents = BrowserWindow->GetResourceContent(Frame, Request);
 			if(Contents.IsSet())
 			{
-				Contents.GetValue().ReplaceInline(TEXT("\n"), TEXT(""), ESearchCase::CaseSensitive);
-				Contents.GetValue().ReplaceInline(TEXT("\r"), TEXT(""), ESearchCase::CaseSensitive);
-
 				// pass the text we'd like to come back as a response to the request post data
 				CefRefPtr<CefPostData> PostData = CefPostData::Create();
 				CefRefPtr<CefPostDataElement> Element = CefPostDataElement::Create();
