Refactor 3D UI display of WKWebKit-based Web Browser for macOS and iOS functionality

diff --git a/Engine/Source/Runtime/WebBrowser/Private/Apple/ApplePlatformWebBrowser.cpp b/Engine/Source/Runtime/WebBrowser/Private/Apple/ApplePlatformWebBrowser.cpp
--- a/Engine/Source/Runtime/WebBrowser/Private/Apple/ApplePlatformWebBrowser.cpp
+++ b/Engine/Source/Runtime/WebBrowser/Private/Apple/ApplePlatformWebBrowser.cpp
@@ -250,9 +250,12 @@ class SAppleWebBrowserWidget : public SLeafWidget
 							FTextureRHIRef VideoTexture = [NativeWebBrowser GetVideoTexture];
 							if (VideoTexture == nullptr)
 							{
+								// We need this texture to be accessible and writable from the CPU because we need to call
+								// replaceRegion:mipmapLevel:withBytes:bytesPerRow: to copy over a screenshot captured in WKWebView
+								// using takeSnapshotWithConfiguration:completionHandler: which gives us an UI/NSImage.
 								const FRHITextureCreateDesc Desc =
 									FRHITextureCreateDesc::Create2D(TEXT("SAppleWebBrowserWidget_VideoTexture"), Params.Size, PF_R8G8B8A8)
-									.SetFlags(ETextureCreateFlags::External);
+									.SetFlags(ETextureCreateFlags::CPUWritable | ETextureCreateFlags::CPUReadback);
 
 								VideoTexture = RHICreateTexture(Desc);
 								[NativeWebBrowser SetVideoTexture : VideoTexture];
@@ -802,8 +805,21 @@ supportsMetal : (bool)InSupportsMetal supportsMetalMRT : (bool)InSupportsMetalMR
 	{
 		if (WebView != nil)
 		{
+#if PLATFORM_MAC
+			if(Is3DBrowser) 
+			{
+				bNeedsAddToView = false;
+			}
+			else
+			{
+				WebViewContainer.frame = self.DesiredFrame;
+				WebView.frame = WebViewContainer.bounds;
+			}
+#else
 			WebViewContainer.frame = self.DesiredFrame;
 			WebView.frame = WebViewContainer.bounds;
+#endif
+			
 			if (bNeedsAddToView)
 			{
 				bNeedsAddToView = false;
@@ -1049,37 +1065,59 @@ supportsMetal : (bool)InSupportsMetal supportsMetalMRT : (bool)InSupportsMetalMR
 -(void)updateWebViewMetalTexture:(id<MTLTexture>)texture
 {
 #if !PLATFORM_TVOS
-	@autoreleasepool {
+	@autoreleasepool 
+	{
+		WKSnapshotConfiguration *SnapshotConfig = [[WKSnapshotConfiguration alloc] init];
 		
-		CGColorSpaceRef ColorSpace = CGColorSpaceCreateDeviceRGB();
 		NSUInteger Width = [texture width];
 		NSUInteger Height = [texture height];
-		CGContextRef WebViewContext = CGBitmapContextCreate(NULL, Width, Height, 8, 4 * Width, ColorSpace, (CGBitmapInfo)kCGImageAlphaPremultipliedLast);
+		CGRect SnapshotRect = CGRectMake(0, 0, Width, Height);
 		
+		WebView.frame = SnapshotRect;
+		SnapshotConfig.rect = SnapshotRect;
+		
+		[WebView takeSnapshotWithConfiguration : SnapshotConfig
+							 completionHandler : 
 #if PLATFORM_MAC
-		WebView.wantsLayer = YES;
-		if(WebView.layer == nil) {
-			return;
-		}
-		[WebView.layer renderInContext: WebViewContext];
+		 ^ void(NSImage *Image, NSError *Error)
 #else
-		UIGraphicsBeginImageContextWithOptions(WebView.frame.size, NO, 1.0f);
-	    [WebView drawViewHierarchyInRect : WebView.bounds afterScreenUpdates : NO];
-		UIImage *Image = UIGraphicsGetImageFromCurrentImageContext();
-		UIGraphicsEndImageContext();
-		CGImage *WebViewCGImage = Image.CGImage;
-		CGContextDrawImage(WebViewContext, CGRectMake(0, 0, Width, Height), WebViewCGImage);
+		 ^ void(UIImage *Image, NSError *Error)
 #endif
-		[texture replaceRegion : MTLRegionMake2D(0, 0, Width, Height)
-			mipmapLevel : 0
-			withBytes : CGBitmapContextGetData(WebViewContext)
-			bytesPerRow : 4 * Width];
-		CGColorSpaceRelease(ColorSpace);
-		CGContextRelease(WebViewContext);
-		
-#if !PLATFORM_MAC
-		Image = nil;
+		{
+			if(Error != nil) 
+			{
+				UE_LOG(LogAppleWebBrowser, Error, TEXT("takeSnapshotWithConfiguration errored with NSError %d!"), Error.code);
+				
+				[SnapshotConfig release];
+				return;
+			}
+			
+			CGColorSpaceRef ColorSpace = CGColorSpaceCreateDeviceRGB();
+			CGContextRef WebViewContext = CGBitmapContextCreate(NULL, Width, Height, 8, 4 * Width, ColorSpace, (CGBitmapInfo)kCGImageAlphaPremultipliedLast);
+			
+#if PLATFORM_MAC
+			NSGraphicsContext *NewGfxContext = [NSGraphicsContext graphicsContextWithCGContext : WebViewContext
+																  flipped : NO];
+			NSGraphicsContext *OldGfxContext = NSGraphicsContext.currentContext;
+			NSGraphicsContext.currentContext = NewGfxContext;
+			
+			[Image drawInRect: SnapshotRect];
+		 
+			NSGraphicsContext.currentContext = OldGfxContext;
+#else
+			CGContextDrawImage(WebViewContext, SnapshotRect, Image.CGImage);
 #endif
+			[texture replaceRegion : MTLRegionMake2D(0, 0, Width, Height)
+				 mipmapLevel : 0
+				 withBytes : CGBitmapContextGetData(WebViewContext)
+				 bytesPerRow : 4 * Width];
+			
+			Image = nil;
+		 
+			CGColorSpaceRelease(ColorSpace);
+			CGContextRelease(WebViewContext);
+			[SnapshotConfig release];
+		}];
 	}
 #endif
 }
