Fixed WKWebView frame positioning/sizing in HiDPI displays and HiDPI <--> LoDPI switching on macOS

diff --git a/Engine/Source/Runtime/WebBrowser/Private/Apple/ApplePlatformWebBrowser.cpp b/Engine/Source/Runtime/WebBrowser/Private/Apple/ApplePlatformWebBrowser.cpp
--- a/Engine/Source/Runtime/WebBrowser/Private/Apple/ApplePlatformWebBrowser.cpp
+++ b/Engine/Source/Runtime/WebBrowser/Private/Apple/ApplePlatformWebBrowser.cpp
@@ -134,9 +134,12 @@ class SAppleWebBrowserWidget : public SLeafWidget
 		TSharedPtr<FGenericWindow> EnclosedWindow = WebBrowserWindowPtr.Pin()->GetParentWindow().Get()->GetNativeWindow();
 		{
 			NSWindow* ParentWindow = (NSWindow*)EnclosedWindow->GetOSWindowHandle();
-			
+			float ScaleFactor = [ParentWindow backingScaleFactor];
 			const FVector2D CocoaPosition = FMacApplication::ConvertSlatePositionToCocoa(WindowPos.X, WindowPos.Y);
 			
+			WindowSize.X /= ScaleFactor;
+			WindowSize.Y /= ScaleFactor;
+			
 			NSRect ParentFrame = [ParentWindow frame];
 			NSRect Rect = NSMakeRect(CocoaPosition.X - ParentFrame.origin.x, (CocoaPosition.Y - ParentFrame.origin.y) - WindowSize.Y, 
 									 FMath::Max(WindowSize.X, 1), FMath::Max(WindowSize.Y, 1));
