fixing WKWebView-related memory leaks in ApplePlatformWebBrowser

diff --git a/Engine/Source/Runtime/WebBrowser/Private/Apple/ApplePlatformWebBrowser.cpp b/Engine/Source/Runtime/WebBrowser/Private/Apple/ApplePlatformWebBrowser.cpp
--- a/Engine/Source/Runtime/WebBrowser/Private/Apple/ApplePlatformWebBrowser.cpp
+++ b/Engine/Source/Runtime/WebBrowser/Private/Apple/ApplePlatformWebBrowser.cpp
@@ -71,19 +71,15 @@ class SAppleWebBrowserWidget : public SLeafWidget
 		bool bEnableFloatingCloseButton = false;
 		GConfig->GetBool(TEXT("Browser"), TEXT("bEnableFloatingCloseButton"), bEnableFloatingCloseButton, GEngineIni);
 
-		TSharedPtr<SAppleWebBrowserWidget> SharedThis(this);
 		WebViewWrapper = [AppleWebViewWrapper alloc];
-		[WebViewWrapper create : SharedThis userAgentApplication : Args._UserAgentApplication.GetNSString() useTransparency : Args._UseTransparency supportsMetal : bSupportsMetal supportsMetalMRT : bSupportsMetalMRT enableFloatingCloseButton : bEnableFloatingCloseButton];
+		[WebViewWrapper create : this userAgentApplication : Args._UserAgentApplication.GetNSString() useTransparency : Args._UseTransparency supportsMetal : bSupportsMetal supportsMetalMRT : bSupportsMetalMRT enableFloatingCloseButton : bEnableFloatingCloseButton];
 
 #if !PLATFORM_TVOS
-		OnSafeFrameChangedEventHandle = FCoreDelegates::OnSafeFrameChangedEvent.AddLambda( [this, WeakThis = TWeakPtr<SAppleWebBrowserWidget>(SharedThis)]()
+		OnSafeFrameChangedEventHandle = FCoreDelegates::OnSafeFrameChangedEvent.AddLambda( [this]()
 										{
-											if (TSharedPtr<SAppleWebBrowserWidget> StrongThis = WeakThis.Pin())
+											if (WebViewWrapper != nil)
 											{
-												if (WebViewWrapper != nil)
-												{
-													[WebViewWrapper didRotate];
-												}
+												[WebViewWrapper didRotate];
 											}
 										});
 
@@ -211,12 +207,6 @@ class SAppleWebBrowserWidget : public SLeafWidget
 									FIntPoint(static_cast<int32_t>(Position.X), static_cast<int32_t>(Position.Y)));
 			
 #if !PLATFORM_TVOS
-#if !UE_BUILD_SHIPPING
-			if (WebViewWrapper.WebView != nil)
-			{
-				WebViewWrapper.WebView.inspectable = CVarWebViewEnableInspectable.GetValueOnAnyThread()? YES : NO;
-			}
-#endif
 			if (Is3DBrowser)
 			{
 				if (WebBrowserTexture)
@@ -357,7 +347,7 @@ class SAppleWebBrowserWidget : public SLeafWidget
 			WebViewWrapper = nil;
 		}
 		WebBrowserWindowPtr.Reset();
-		
+
 #if !PLATFORM_TVOS
 		delete TextureSamplePool;
 		TextureSamplePool = nullptr;
@@ -627,11 +617,12 @@ private:
 @synthesize CloseButton;
 #endif
 @synthesize WebViewContainer;
+@synthesize UserContentController;
 #endif
 @synthesize NextURL;
 @synthesize NextContent;
 
--(void)create:(TSharedPtr<SAppleWebBrowserWidget>)InWebBrowserWidget userAgentApplication: (NSString*)UserAgentApplication useTransparency : (bool)InUseTransparency
+-(void)create:(SAppleWebBrowserWidget*)InWebBrowserWidget userAgentApplication: (NSString*)UserAgentApplication useTransparency : (bool)InUseTransparency
 supportsMetal : (bool)InSupportsMetal supportsMetalMRT : (bool)InSupportsMetalMRT enableFloatingCloseButton : (bool)bEnableFloatingCloseButton;
 {
 	WebBrowserWidget = InWebBrowserWidget;
@@ -648,34 +639,38 @@ supportsMetal : (bool)InSupportsMetal supportsMetalMRT : (bool)InSupportsMetalMR
 	dispatch_async(dispatch_get_main_queue(), ^
 	{
 		#if PLATFORM_IOS
-		WebViewContainer = [[UIView alloc]initWithFrame:CGRectMake(1, 1, 100, 100)];
+		self.WebViewContainer = [[UIView alloc]initWithFrame:CGRectMake(1, 1, 100, 100)];
 		[self.WebViewContainer setOpaque : NO];
 		[self.WebViewContainer setBackgroundColor : [UIColor clearColor]];
 		#else
-		WebViewContainer = [[NSView alloc]initWithFrame:CGRectMake(1, 1, 100, 100)];
+		self.WebViewContainer = [[NSView alloc]initWithFrame:CGRectMake(1, 1, 100, 100)];
 		#endif
 
 		WKWebViewConfiguration* theConfiguration = [[WKWebViewConfiguration alloc] init];
+
+		self.UserContentController = [[WKUserContentController alloc] init];
 		
 		theConfiguration.websiteDataStore = [WKWebsiteDataStore dataStoreForIdentifier: FAppleCookieManager::CookieManagerIdentifier];
 		
 		[theConfiguration.preferences setValue:@YES forKey:@"allowFileAccessFromFileURLs"];
 		[theConfiguration.preferences setValue:@YES forKey:@"developerExtrasEnabled"];
-		
+	
 		NSString* MessageHandlerName = [NSString stringWithFString : FMobileJSScripting::JSMessageHandler];
+		[self.UserContentController addScriptMessageHandler:self name: MessageHandlerName];
+		
 		theConfiguration.applicationNameForUserAgent = UserAgentApplication;
-		[theConfiguration.userContentController addScriptMessageHandler:self name: MessageHandlerName];
+		theConfiguration.userContentController = self.UserContentController;
 
-		WebView = [[WKWebView alloc]initWithFrame:CGRectMake(1, 1, 100, 100)configuration : theConfiguration];
+		self.WebView = [[WKWebView alloc]initWithFrame:CGRectMake(1, 1, 100, 100)configuration : theConfiguration];
 #if !UE_BUILD_SHIPPING
-		WebView.inspectable = CVarWebViewEnableInspectable.GetValueOnAnyThread()? YES : NO;
+		self.WebView.inspectable = CVarWebViewEnableInspectable.GetValueOnAnyThread()? YES : NO;
 #endif
 		[self.WebViewContainer addSubview : WebView];
-		WebView.navigationDelegate = self;
-		WebView.UIDelegate = self;
+		self.WebView.navigationDelegate = self;
+		self.WebView.UIDelegate = self;
 
 #if PLATFORM_IOS
-		WebView.scrollView.bounces = NO;
+		self.WebView.scrollView.bounces = NO;
 		
 		if (InUseTransparency)
 		{
@@ -692,11 +687,19 @@ supportsMetal : (bool)InSupportsMetal supportsMetalMRT : (bool)InSupportsMetalMR
 			CloseButton = MakeCloseButton();
 			[self.WebView addSubview : CloseButton];
 
-			TWeakPtr<SAppleWebBrowserWidget> WeakWebBrowserWidget = InWebBrowserWidget;
-			CloseButton.TapHandler = [^{
-				if (TSharedPtr<SAppleWebBrowserWidget> StrongWebBrowserWidget = WeakWebBrowserWidget.Pin())
+			CloseButton.TapHandler = [^
 				{
-					StrongWebBrowserWidget->FloatingCloseButtonPressed();
+				// WebBrowserWidget should have the same lifecycle as this
+				// CloseButton (see [AppleWebViewWrapper cleanup], 
+				// which is only called by SAppleWebBrowserWidget::Close()), 
+				// so using this raw ptr should be safe. Previously, a
+				// shared ptr was used, but the shared ptr was created in the
+				// SAppleWebBrowserWidget's Construct function, 
+				// which caused strange behavior with the TSharedRef SNew creates,
+				// potentially leading to over-releasing.
+				if(WebBrowserWidget != nil)
+				{
+					WebBrowserWidget->FloatingCloseButtonPressed();
 				}
 			} copy];
 		}
@@ -720,29 +723,58 @@ supportsMetal : (bool)InSupportsMetal supportsMetalMRT : (bool)InSupportsMetalMR
 #endif
 }
 
--(void)close;
+-(void)cleanup;
 {
 #if !PLATFORM_TVOS
-	WebView.navigationDelegate = nil;
-	dispatch_async(dispatch_get_main_queue(), ^
+	if(self.UserContentController != nil) 
 	{
-		[self.WebViewContainer removeFromSuperview];
-		[self.WebView removeFromSuperview];
+		[self.UserContentController removeAllScriptMessageHandlers];
+		[self.UserContentController removeAllUserScripts];
 		
-		#if PLATFORM_IOS
-		if (CloseButton)
+		[self.UserContentController release];
+ 
+		self.UserContentController = nil;
+	}
+	
+	if (self.WebView != nil)
+	{		
+#if PLATFORM_MAC
+		if(self.WebView.window.firstResponder == self.WebView) 
 		{
-			[CloseButton release];
+			[self.WebView.window makeFirstResponder: nil];
 		}
+#endif
+		
+		self.WebView.navigationDelegate = nil;
+		self.WebView.UIDelegate = nil;
+		
+		[self.WebView removeFromSuperview];
+		[self.WebViewContainer removeFromSuperview];
+		
+		[self.WebView release];
+		[self.WebViewContainer release];
+		
+		self.WebView = nil;
+		self.WebViewContainer = nil;
+		WebBrowserWidget = nil;
+	}
+
+#if !PLATFORM_MAC
+	if (CloseButton != nil)
+	{
+		[CloseButton release];
 		CloseButton = nil;
-		#endif
-		[WebView release];
-		[WebViewContainer release];
+	}
+#endif
+#endif
+}
 
-		WebView = nil;
-		WebViewContainer = nil;
+-(void)close;
+{
+	dispatch_async(dispatch_get_main_queue(), ^
+	{
+		[self cleanup];
 	});
-#endif
 }
 
 -(void)showFloatingCloseButton:(BOOL)bShow setDraggable:(BOOL)bDraggable;
@@ -764,25 +796,7 @@ supportsMetal : (bool)InSupportsMetal supportsMetalMRT : (bool)InSupportsMetalMR
 
 -(void)dealloc;
 {
-#if !PLATFORM_TVOS
-	if (WebView != nil)
-	{
-		WebView.navigationDelegate = nil;
-		[WebView release];
-		WebView = nil;
-	};
-
-#if !PLATFORM_MAC
-	if (CloseButton)
-	{
-		[CloseButton release];
-		CloseButton = nil;
-	};
-#endif
-
-	[WebViewContainer release];
-	WebViewContainer = nil;
-#endif
+	[self cleanup];
 	[NextContent release];
 	NextContent = nil;
 	[NextURL release];
@@ -806,8 +820,13 @@ supportsMetal : (bool)InSupportsMetal supportsMetalMRT : (bool)InSupportsMetalMR
 #if !PLATFORM_TVOS
 	dispatch_async(dispatch_get_main_queue(), ^
 	{
-		if (WebView != nil)
+		if (self.WebView != nil)
 		{
+			
+#if !UE_BUILD_SHIPPING
+			self.WebView.inspectable = CVarWebViewEnableInspectable.GetValueOnAnyThread()? YES : NO;
+#endif
+			
 #if PLATFORM_MAC
 			if(Is3DBrowser) 
 			{
@@ -815,14 +834,13 @@ supportsMetal : (bool)InSupportsMetal supportsMetalMRT : (bool)InSupportsMetalMR
 			}
 			else
 			{
-				WebViewContainer.frame = self.DesiredFrame;
-				WebView.frame = WebViewContainer.bounds;
+				self.WebViewContainer.frame = self.DesiredFrame;
+				self.WebView.frame = WebViewContainer.bounds;
 			}
 #else
-			WebViewContainer.frame = self.DesiredFrame;
-			WebView.frame = WebViewContainer.bounds;
+			self.WebViewContainer.frame = self.DesiredFrame;
+			self.WebView.frame = WebViewContainer.bounds;
 #endif
-			
 			if (bNeedsAddToView)
 			{
 				bNeedsAddToView = false;
@@ -1250,9 +1268,9 @@ TSharedRef<SWidget> FWebBrowserWindow::CreateWidget()
 		.WebBrowserWindow(SharedThis(this));
 
 	BrowserWidget = BrowserWidgetRef;
-
-	Scripting->SetWindow(SharedThis(this));
 	
+	Scripting->SetWindow(SharedThis(this));
+
 	return BrowserWidgetRef;
 }
 
diff --git a/Engine/Source/Runtime/WebBrowser/Private/Apple/ApplePlatformWebBrowser.h b/Engine/Source/Runtime/WebBrowser/Private/Apple/ApplePlatformWebBrowser.h
--- a/Engine/Source/Runtime/WebBrowser/Private/Apple/ApplePlatformWebBrowser.h
+++ b/Engine/Source/Runtime/WebBrowser/Private/Apple/ApplePlatformWebBrowser.h
@@ -43,7 +43,7 @@ class SWebBrowserView;
 @interface AppleWebViewWrapper : NSObject
 #endif
 {
-	TSharedPtr<SAppleWebBrowserWidget> WebBrowserWidget;
+	SAppleWebBrowserWidget* WebBrowserWidget;
 	FTextureRHIRef VideoTexture;
 	bool bNeedsAddToView;
 	bool Is3DBrowser;
@@ -59,12 +59,13 @@ class SWebBrowserView;
 @property(strong) WebViewCloseButton* CloseButton;
 @property(strong) UIView* WebViewContainer;
 #endif
+@property(strong) WKUserContentController* UserContentController;
 #endif
 @property(copy) NSURL* NextURL;
 @property(copy) NSString* NextContent;
 @property CGRect DesiredFrame;
 
--(void)create:(TSharedPtr<SAppleWebBrowserWidget>)InWebBrowserWidget userAgentApplication: (NSString*)UserAgentApplication useTransparency : (bool)InUseTransparency
+-(void)create:(SAppleWebBrowserWidget*)InWebBrowserWidget userAgentApplication: (NSString*)UserAgentApplication useTransparency : (bool)InUseTransparency
 supportsMetal : (bool)InSupportsMetal supportsMetalMRT : (bool)InSupportsMetalMRT enableFloatingCloseButton : (bool)bEnableFloatingCloseButton;
 -(void)didRotate;
 -(void)close;
