Adapting MobileJS for macOS to get Fab working with WebKit on macOS Editor

diff --git a/Engine/Plugins/Fab/Source/Fab/Private/FabBrowser.cpp b/Engine/Plugins/Fab/Source/Fab/Private/FabBrowser.cpp
--- a/Engine/Plugins/Fab/Source/Fab/Private/FabBrowser.cpp
+++ b/Engine/Plugins/Fab/Source/Fab/Private/FabBrowser.cpp
@@ -295,6 +295,9 @@ TSharedRef<SDockTab> FFabBrowser::OpenTab(const FSpawnTabArgs&)
 	FString FinalUrl                = FPaths::Combine(TEXT("file:///"), IndexUrl);
 	WindowSettings.InitialURL       = FinalUrl;
 	WindowSettings.BrowserFrameRate = 60;
+#if PLATFORM_MAC
+	WindowSettings.bMobileJSReturnInDict = false;
+#endif
 
 	IWebBrowserSingleton* WebBrowserSingleton = WebBrowserModule.GetSingleton();
 	WebBrowserSingleton->SetDevToolsShortcutEnabled(true);
@@ -420,7 +423,7 @@ void FFabBrowser::Shutdown()
 
 void FFabBrowser::LoggedIn(const FString& InAccessToken)
 {
-	ExecuteJavascript(FString::Printf(TEXT("window.ue.fab.onLoginSuccessful('%s');"), *InAccessToken));
+	ExecuteJavascript(FString::Printf(TEXT("window.ue.fab.onLoginSuccessful('%s')"), *InAccessToken));
 }
 
 void FFabBrowser::GetSignedUrl(const FString& AssetId, const int32 Tier)
@@ -438,21 +441,24 @@ void FFabBrowser::LogEvent(const FFabAnalyticsPayload& Payload)
 
 void FFabBrowser::ShowSettings()
 {
-	const TSharedRef<SWindow> Window = SNew(SWindow)
-		.Title(LOCTEXT("FabSettingsLabel", "Fab Settings"))
-		.ClientSize(FVector2D(600.f, 300.f))
-		.SizingRule(ESizingRule::UserSized);
-
-	TSharedPtr<SFabSettingsWindow> SettingsWindow;
-	Window->SetContent(SAssignNew(SettingsWindow, SFabSettingsWindow).WidgetWindow(Window));
-	TSharedPtr<SWindow> ParentWindow;
-	if (FModuleManager::Get().IsModuleLoaded("MainFrame"))
-	{
-		const IMainFrameModule& MainFrame = FModuleManager::LoadModuleChecked<IMainFrameModule>("MainFrame");
-		ParentWindow                      = MainFrame.GetParentWindow();
-	}
+	AsyncTask(ENamedThreads::Type::GameThread, []() 
+    {
+		const TSharedRef<SWindow> Window = SNew(SWindow)
+			.Title(LOCTEXT("FabSettingsLabel", "Fab Settings"))
+			.ClientSize(FVector2D(600.f, 300.f))
+			.SizingRule(ESizingRule::UserSized);
+
+		TSharedPtr<SFabSettingsWindow> SettingsWindow;
+		Window->SetContent(SAssignNew(SettingsWindow, SFabSettingsWindow).WidgetWindow(Window));
+		TSharedPtr<SWindow> ParentWindow;
+		if (FModuleManager::Get().IsModuleLoaded("MainFrame"))
+		{
+			const IMainFrameModule& MainFrame = FModuleManager::LoadModuleChecked<IMainFrameModule>("MainFrame");
+			ParentWindow                      = MainFrame.GetParentWindow();
+		}
 
-	FSlateApplication::Get().AddModalWindow(Window, ParentWindow, false);
+		FSlateApplication::Get().AddModalWindow(Window, ParentWindow, false);
+	});
 }
 
 void FFabBrowser::OpenURL(const FString& InURL)
diff --git a/Engine/Plugins/Fab/Source/Fab/Private/FabBrowserApi.cpp b/Engine/Plugins/Fab/Source/Fab/Private/FabBrowserApi.cpp
--- a/Engine/Plugins/Fab/Source/Fab/Private/FabBrowserApi.cpp
+++ b/Engine/Plugins/Fab/Source/Fab/Private/FabBrowserApi.cpp
@@ -69,7 +69,12 @@ void UFabBrowserApi::AddToProject(const FString& DownloadUrl, const FFabAssetMet
 				CompleteWorkflow(AssetId);
 			}
 		);
-		PackImportWorkflow->Execute();
+		
+		AsyncTask(ENamedThreads::Type::GameThread, [PackImportWorkflow]() 
+	    {
+			PackImportWorkflow->Execute();
+		});
+		
 		this->ActiveWorkflows.Add(PackImportWorkflow);
 		return;
 	}
@@ -262,48 +267,48 @@ FFabFrontendSettings UFabBrowserApi::GetSettings()
 	FFabFrontendSettings FrontendSettings;
 	if (FabSettings->PreferredDefaultFormat == EFabPreferredFormats::GLTF)
 	{
-		FrontendSettings.PreferredFormat = "gltf";
+		FrontendSettings.preferredformat = "gltf";
 	}
 	else if (FabSettings->PreferredDefaultFormat == EFabPreferredFormats::FBX)
 	{
-		FrontendSettings.PreferredFormat = "fbx";
+		FrontendSettings.preferredformat = "fbx";
 	}
 	if (FabSettings->PreferredQualityTier == EFabPreferredQualityTier::Low)
 	{
-		FrontendSettings.PreferredQuality = "low";
+		FrontendSettings.preferredquality = "low";
 	}
 	else if (FabSettings->PreferredQualityTier == EFabPreferredQualityTier::Medium)
 	{
-		FrontendSettings.PreferredQuality = "medium";
+		FrontendSettings.preferredquality = "medium";
 	}
 	else if (FabSettings->PreferredQualityTier == EFabPreferredQualityTier::High)
 	{
-		FrontendSettings.PreferredQuality = "high";
+		FrontendSettings.preferredquality = "high";
 	}
 	else if (FabSettings->PreferredQualityTier == EFabPreferredQualityTier::Raw)
 	{
-		FrontendSettings.PreferredQuality = "raw";
+		FrontendSettings.preferredquality = "raw";
 	}
 
 	return FrontendSettings;
 }
 
-void UFabBrowserApi::SetPreferredQualityTier(const FString& PreferredQuality)
+void UFabBrowserApi::SetPreferredQualityTier(const FString& preferredquality)
 {
 	UFabSettings* FabSettings = GetMutableDefault<UFabSettings>();
-	if (PreferredQuality == "low")
+	if (preferredquality == "low")
 	{
 		FabSettings->PreferredQualityTier = EFabPreferredQualityTier::Low;
 	}
-	else if (PreferredQuality == "medium")
+	else if (preferredquality == "medium")
 	{
 		FabSettings->PreferredQualityTier = EFabPreferredQualityTier::Medium;
 	}
-	else if (PreferredQuality == "high")
+	else if (preferredquality == "high")
 	{
 		FabSettings->PreferredQualityTier = EFabPreferredQualityTier::High;
 	}
-	else if (PreferredQuality == "raw")
+	else if (preferredquality == "raw")
 	{
 		FabSettings->PreferredQualityTier = EFabPreferredQualityTier::Raw;
 	}
diff --git a/Engine/Plugins/Fab/Source/Fab/Private/FabBrowserApi.h b/Engine/Plugins/Fab/Source/Fab/Private/FabBrowserApi.h
--- a/Engine/Plugins/Fab/Source/Fab/Private/FabBrowserApi.h
+++ b/Engine/Plugins/Fab/Source/Fab/Private/FabBrowserApi.h
@@ -31,11 +31,14 @@ struct FFabFrontendSettings
 {
 	GENERATED_BODY()
 
+	// These two uproperties need to be in this case (NOT PascalCase)
+	// otherwise we will break compatibility with the engine plugin 
+	// Fab's JS on the macOS WebKit WebBrowser plugin backend.
 	UPROPERTY()
-	FString PreferredFormat;
+	FString preferredformat;
 
 	UPROPERTY()
-	FString PreferredQuality;
+	FString preferredquality;
 };
 
 UCLASS()
diff --git a/Engine/Source/Runtime/WebBrowser/Private/Android/AndroidWebBrowserWindow.cpp b/Engine/Source/Runtime/WebBrowser/Private/Android/AndroidWebBrowserWindow.cpp
--- a/Engine/Source/Runtime/WebBrowser/Private/Android/AndroidWebBrowserWindow.cpp
+++ b/Engine/Source/Runtime/WebBrowser/Private/Android/AndroidWebBrowserWindow.cpp
@@ -23,7 +23,7 @@ namespace {
 
 }
 
-FAndroidWebBrowserWindow::FAndroidWebBrowserWindow(FString InUrl, TOptional<FString> InContentsToLoad, bool InShowErrorMessage, bool InThumbMouseButtonNavigation, bool InUseTransparency, bool bInJSBindingToLoweringEnabled, const FString& InUserAgentApplication)
+FAndroidWebBrowserWindow::FAndroidWebBrowserWindow(FString InUrl, TOptional<FString> InContentsToLoad, bool InShowErrorMessage, bool InThumbMouseButtonNavigation, bool InUseTransparency, bool bInJSBindingToLoweringEnabled, const FString& InUserAgentApplication, bool bMobileJSReturnInDict)
 	: CurrentUrl(MoveTemp(InUrl))
 	, ContentsToLoad(MoveTemp(InContentsToLoad))
 	, UserAgentApplication(InUserAgentApplication)
@@ -37,7 +37,7 @@ FAndroidWebBrowserWindow::FAndroidWebBrowserWindow(FString InUrl, TOptional<FStr
 {
 	bool bInjectJSOnPageStarted = false;
 	GConfig->GetBool(TEXT("Browser"), TEXT("bInjectJSOnPageStarted"), bInjectJSOnPageStarted, GEngineIni);
-	Scripting = MakeShared<FMobileJSScripting>(bInJSBindingToLoweringEnabled, bInjectJSOnPageStarted);
+	Scripting = MakeShared<FMobileJSScripting>(bInJSBindingToLoweringEnabled, bInjectJSOnPageStarted, bMobileJSReturnInDict);
 	// Deal with optional texture resolution override
 	FString WebViewTextureSize;
 	if (GConfig->GetString(TEXT("/Script/AndroidRuntimeSettings.AndroidRuntimeSettings"), TEXT("WebViewTextureSize"), WebViewTextureSize, GEngineIni))
diff --git a/Engine/Source/Runtime/WebBrowser/Private/Android/AndroidWebBrowserWindow.h b/Engine/Source/Runtime/WebBrowser/Private/Android/AndroidWebBrowserWindow.h
--- a/Engine/Source/Runtime/WebBrowser/Private/Android/AndroidWebBrowserWindow.h
+++ b/Engine/Source/Runtime/WebBrowser/Private/Android/AndroidWebBrowserWindow.h
@@ -36,8 +36,9 @@ private:
 	 * @param bThumbMouseButtonNavigation Whether to allow forward and back navigation via the mouse thumb buttons.
 	 * @param bUseTransparency Whether to enable transparency.
 	 * @param bJSBindingToLoweringEnabled Whether we ToLower all JavaScript member names.
+	 * @param bMobileJSReturnInDict Whether to pass a UObject wrapper function's return value in a dictionary with the property's name or not.
 	 */
-	FAndroidWebBrowserWindow(FString Url, TOptional<FString> ContentsToLoad, bool bShowErrorMessage, bool bThumbMouseButtonNavigation, bool bUseTransparency, bool bJSBindingToLoweringEnabled, const FString& InUserAgentApplication);
+	FAndroidWebBrowserWindow(FString Url, TOptional<FString> ContentsToLoad, bool bShowErrorMessage, bool bThumbMouseButtonNavigation, bool bUseTransparency, bool bJSBindingToLoweringEnabled, const FString& InUserAgentApplication, bool bMobileJSReturnInDict);
 
 	/**
 	 * Create the SWidget for this WebBrowserWindow
diff --git a/Engine/Source/Runtime/WebBrowser/Private/Apple/ApplePlatformWebBrowser.cpp b/Engine/Source/Runtime/WebBrowser/Private/Apple/ApplePlatformWebBrowser.cpp
--- a/Engine/Source/Runtime/WebBrowser/Private/Apple/ApplePlatformWebBrowser.cpp
+++ b/Engine/Source/Runtime/WebBrowser/Private/Apple/ApplePlatformWebBrowser.cpp
@@ -916,10 +916,13 @@ supportsMetal : (bool)InSupportsMetal supportsMetalMRT : (bool)InSupportsMetalMR
 		
 		//NSLog(@"executejavascript %@", InJavaScript);
 		
-		[self.WebView evaluateJavaScript : InJavaScript completionHandler : ^(id Object, NSError* Error) {
-			if(Error != nil) {
-				UE_LOG(LogAppleWebBrowser, Error, TEXT("executeJavaScript(%s) errored with NSError %d!"), 
-					   *InJavaScriptFString, Error.code);
+		[self.WebView evaluateJavaScript : InJavaScript completionHandler : ^(id Object, NSError* Error) 
+		 {
+			if(Error != nil) 
+			{
+				const char *domainString = [Error.domain UTF8String];
+				UE_LOG(LogAppleWebBrowser, Error, TEXT("executeJavaScript(\"%s\") errored with NSError %d in domain %hs!"), 
+					   *InJavaScriptFString, Error.code, domainString);
 			}
 		}];
 	});
@@ -1226,7 +1229,7 @@ namespace {
 
 }
 
-FWebBrowserWindow::FWebBrowserWindow(FString InUrl, TOptional<FString> InContentsToLoad, bool InShowErrorMessage, bool InThumbMouseButtonNavigation, bool InUseTransparency, bool bInJSBindingToLoweringEnabled, FString InUserAgentApplication)
+FWebBrowserWindow::FWebBrowserWindow(FString InUrl, TOptional<FString> InContentsToLoad, bool InShowErrorMessage, bool InThumbMouseButtonNavigation, bool InUseTransparency, bool bInJSBindingToLoweringEnabled, FString InUserAgentApplication, bool bMobileJSReturnInDict)
 	: CurrentUrl(MoveTemp(InUrl))
 	, UserAgentApplication(MoveTemp(InUserAgentApplication))
 	, ContentsToLoad(MoveTemp(InContentsToLoad))
@@ -1240,7 +1243,8 @@ FWebBrowserWindow::FWebBrowserWindow(FString InUrl, TOptional<FString> InContent
 {
 	bool bInjectJSOnPageStarted = false;
 	GConfig->GetBool(TEXT("Browser"), TEXT("bInjectJSOnPageStarted"), bInjectJSOnPageStarted, GEngineIni);
-	Scripting = MakeShared<FMobileJSScripting>(bInJSBindingToLoweringEnabled, bInjectJSOnPageStarted);
+	
+	Scripting = MakeShared<FMobileJSScripting>(bInJSBindingToLoweringEnabled, bInjectJSOnPageStarted, bMobileJSReturnInDict);
 }
 
 FWebBrowserWindow::~FWebBrowserWindow()
diff --git a/Engine/Source/Runtime/WebBrowser/Private/Apple/ApplePlatformWebBrowser.h b/Engine/Source/Runtime/WebBrowser/Private/Apple/ApplePlatformWebBrowser.h
--- a/Engine/Source/Runtime/WebBrowser/Private/Apple/ApplePlatformWebBrowser.h
+++ b/Engine/Source/Runtime/WebBrowser/Private/Apple/ApplePlatformWebBrowser.h
@@ -117,7 +117,7 @@ private:
 	* @param InContentsToLoad Optional string to load as a web page.
 	* @param InShowErrorMessage Whether to show an error message in case of loading errors.
 	*/
-	FWebBrowserWindow(FString InUrl, TOptional<FString> InContentsToLoad, bool ShowErrorMessage, bool bThumbMouseButtonNavigation, bool bUseTransparency, bool bInJSBindingToLoweringEnabled, FString InUserAgentApplication);
+	FWebBrowserWindow(FString InUrl, TOptional<FString> InContentsToLoad, bool ShowErrorMessage, bool bThumbMouseButtonNavigation, bool bUseTransparency, bool bInJSBindingToLoweringEnabled, FString InUserAgentApplication, bool bMobileJSReturnInDict);
 
 	/**
 	 * Create the SWidget for this WebBrowserWindow
diff --git a/Engine/Source/Runtime/WebBrowser/Private/MobileJS/MobileJSScripting.cpp b/Engine/Source/Runtime/WebBrowser/Private/MobileJS/MobileJSScripting.cpp
--- a/Engine/Source/Runtime/WebBrowser/Private/MobileJS/MobileJSScripting.cpp
+++ b/Engine/Source/Runtime/WebBrowser/Private/MobileJS/MobileJSScripting.cpp
@@ -17,6 +17,8 @@
 // For UrlDecode/Encode
 #include "Http.h"
 
+DEFINE_LOG_CATEGORY(LogMobileJSScripting);
+
 // Inseterted as a part of an URL to send a message to the front end.
 // Note, we can't use a custom protocol due to cross-domain issues.
 const FString FMobileJSScripting::JSMessageTag = TEXT("/!!com.epicgames.unreal.message/");
@@ -95,7 +97,7 @@ namespace
 			// encodes and sends a message to the host application
 			TEXT("sendMessage: function()")
 			TEXT("{")
-#if PLATFORM_IOS
+#if PLATFORM_IOS || PLATFORM_MAC
 			TEXT("	window.webkit.messageHandlers.") + FMobileJSScripting::JSMessageHandler + TEXT(".postMessage(Array.prototype.map.call(arguments,function(e){return encodeURIComponent(e)}).join('/'));")
 #else
 			TEXT("	var req=new XMLHttpRequest();")
@@ -324,7 +326,10 @@ void FMobileJSScripting::BindUObject(const TSharedRef<class IWebBrowserWindow>&
 
 	InitializeScript(InWindow);
 	FString SetValueScript = FString::Printf(TEXT("window.ue['%s'] = %s;"), *ExposedName.ReplaceCharWithEscapedChar(), *Converted);
+	SetValueScript.Append(ScriptingPostInit);
+	
 	InWindow->ExecuteJavascript(SetValueScript);
+	InWindow->ExecuteJavascript(ScriptingPostInit);
 }
 
 void FMobileJSScripting::UnbindUObject(const TSharedRef<class IWebBrowserWindow>& InWindow, const FString& Name, UObject* Object, bool bIsPermanent)
@@ -347,6 +352,10 @@ bool FMobileJSScripting::OnJsMessageReceived(const FString& Command, const TArra
 	if (Command == ExecuteMethodCommand)
 	{
 		Result = HandleExecuteUObjectMethodMessage(Params);
+	} 
+	else 
+	{
+		UE_LOG(LogMobileJSScripting, Error, TEXT("Unknown JS message received: %s"), *Command);
 	}
 
 	return Result;
@@ -372,6 +381,7 @@ FString FMobileJSScripting::ConvertObject(UObject* Object)
 	RetainBinding(Object);
 	UClass* Class = Object->GetClass();
 
+	FString ObjectGuidString = *PtrToGuid(Object).ToString(EGuidFormats::Digits);
 	bool first = true;
 	FString Result = TEXT("(function(){ return Object.create({");
 	for (TFieldIterator<UFunction> FunctionIt(Class, EFieldIteratorFlags::IncludeSuper); FunctionIt; ++FunctionIt)
@@ -413,11 +423,13 @@ FString FMobileJSScripting::ConvertObject(UObject* Object)
 		}
 
 		Result.Append(TEXT(")"));
-		Result.Append(TEXT(" {return window.ue.$.executeMethod(this.$id, arguments)}"));
+		Result.Append(TEXT(" {return window.ue.$.executeMethod(\""));
+		Result.Append(ObjectGuidString);
+		Result.Append(TEXT("\", arguments)}"));
 	}
 	Result.Append(TEXT("},{"));
 	Result.Append(TEXT("$id: {writable: false, configurable:false, enumerable: false, value: '"));
- 	Result.Append(*PtrToGuid(Object).ToString(EGuidFormats::Digits));
+	Result.Append(ObjectGuidString);
 	Result.Append(TEXT("'}})})()"));
 	return Result;
 }
@@ -459,6 +471,7 @@ void FMobileJSScripting::InvokeJSFunctionRaw(FGuid FunctionId, const FString& Ra
 void FMobileJSScripting::InvokeJSErrorResult(FGuid FunctionId, const FString& Error)
 {
 	FWebJSParam Args[1] = {FWebJSParam(Error)};
+	UE_LOG(LogMobileJSScripting, Error, TEXT("JS Error %s"), *Error);
 	InvokeJSFunction(FunctionId, 1, Args, true);
 }
 
@@ -466,6 +479,7 @@ bool FMobileJSScripting::HandleExecuteUObjectMethodMessage(const TArray<FString>
 {
 	if (MessageArgs.Num() != 4)
 	{
+		UE_LOG(LogMobileJSScripting, Error, TEXT("insufficent num of args expected 4 got %d"), MessageArgs.Num());
 		return false;
 	}
 
@@ -473,6 +487,7 @@ bool FMobileJSScripting::HandleExecuteUObjectMethodMessage(const TArray<FString>
 	if (!FGuid::Parse(MessageArgs[0], ObjectKey))
 	{
 		// Invalid GUID
+		UE_LOG(LogMobileJSScripting, Error, TEXT("JS GUID %s was invalid"), *MessageArgs[0]);
 		return false;
 	}
 	// Get the promise callback and use that to report any results from executing this function.
@@ -480,6 +495,7 @@ bool FMobileJSScripting::HandleExecuteUObjectMethodMessage(const TArray<FString>
 	if (!FGuid::Parse(MessageArgs[1], ResultCallbackId))
 	{
 		// Invalid GUID
+		UE_LOG(LogMobileJSScripting, Error, TEXT("JS GUID %s was invalid"), *MessageArgs[1]);
 		return false;
 	}
 
@@ -564,12 +580,37 @@ bool FMobileJSScripting::HandleExecuteUObjectMethodMessage(const TArray<FString>
 
 			// Extract the result value from the serialized JSON object:
 			FString ResultJS = ReturnBackend.ToString();
-			/*
-			ResultJS.Append(TEXT("['"));
-			ResultJS.Append(GetBindingName(ReturnParam).ReplaceCharWithEscapedChar());
-			ResultJS.Append(TEXT("']"));
-			*/
-			InvokeJSFunctionRaw(ResultCallbackId, ResultJS, false);
+			
+			if(!bDefaultJSReturnInDict) 
+			{
+				FString ParsedJS;
+				TSharedPtr<FJsonObject> JsonObj = MakeShareable(new FJsonObject);
+							
+				TSharedPtr<TJsonWriter<>> Writer = TJsonWriterFactory<>::Create(&ParsedJS); 
+				TSharedRef<TJsonReader<>> Reader = TJsonReaderFactory<>::Create(ResultJS);
+
+				if (FJsonSerializer::Deserialize(Reader, JsonObj))
+				{
+					FString ReturnValueFieldName = TEXT("ReturnValue");
+					
+					TArray<TSharedPtr<FJsonValue>> ValuesArray;
+					TSharedPtr<FJsonValue> ReturnValueField = JsonObj->TryGetField(ReturnValueFieldName);
+					
+					ValuesArray.Add(ReturnValueField);
+					
+					FJsonSerializer::Serialize(ValuesArray, *Writer, true);
+				}
+				
+				ParsedJS.RemoveFromStart("[\n", ESearchCase::Type::CaseSensitive);
+				ParsedJS.RemoveFromEnd("\n]", ESearchCase::Type::CaseSensitive);
+				
+				ParsedJS.TrimStartInline();
+				ParsedJS.TrimEndInline();
+				
+				InvokeJSFunctionRaw(ResultCallbackId, ParsedJS, false);
+			} else {
+				InvokeJSFunctionRaw(ResultCallbackId, ResultJS, false);
+			}
 		}
 		else
 		{
@@ -598,9 +639,10 @@ void FMobileJSScripting::InitializeScript(const TSharedRef<class IWebBrowserWind
 		ScreenHeight = FCString::Atoi(**ScreenHeightVar);
 	}
 	Script.Append(TEXT("window.uePlatform = \"Android\";\n"));
-#endif
-#if PLATFORM_IOS
+#elif PLATFORM_IOS
 	Script.Append(TEXT("window.uePlatform = \"iOS\";\n"));
+#elif PLATFORM_MAC
+	Script.Append(TEXT("window.uePlatform = \"macOS\";\n"));
 #endif
 	Script.Append(*FString::Printf(TEXT(
 		"window.ueDeviceWidth = %d;\n"
@@ -650,12 +692,14 @@ void FMobileJSScripting::InjectJavascript(const TSharedRef<class IWebBrowserWind
 	{
 		Script.Append(*FString::Printf(TEXT("window.ue['%s'] = %s;"), *Item.Key.ReplaceCharWithEscapedChar(), *ConvertObject(Item.Value)));
 	}
+	Script.Append(ScriptingPostInit);
 	Window->ExecuteJavascript(Script);
 }
 
-FMobileJSScripting::FMobileJSScripting(bool bJSBindingToLoweringEnabled, bool bInjectJSOnPageStarted)
+FMobileJSScripting::FMobileJSScripting(bool bJSBindingToLoweringEnabled, bool bInjectJSOnPageStarted, bool bDefaultJSReturnInDict)
 	: FWebJSScripting(bJSBindingToLoweringEnabled)
 	, bInjectJSOnPageStarted(bInjectJSOnPageStarted)
+	, bDefaultJSReturnInDict(bDefaultJSReturnInDict)
 {
 }
 
diff --git a/Engine/Source/Runtime/WebBrowser/Private/MobileJS/MobileJSScripting.h b/Engine/Source/Runtime/WebBrowser/Private/MobileJS/MobileJSScripting.h
--- a/Engine/Source/Runtime/WebBrowser/Private/MobileJS/MobileJSScripting.h
+++ b/Engine/Source/Runtime/WebBrowser/Private/MobileJS/MobileJSScripting.h
@@ -3,11 +3,13 @@
 
 #include "CoreMinimal.h"
 
-#if PLATFORM_ANDROID  || PLATFORM_IOS  || PLATFORM_MAC
+#if PLATFORM_ANDROID  || PLATFORM_IOS  ||  PLATFORM_MAC
 
 #include "WebJSFunction.h"
 #include "WebJSScripting.h"
 
+DECLARE_LOG_CATEGORY_EXTERN(LogMobileJSScripting, Log, All);
+
 typedef TSharedRef<class FMobileJSScripting> FMobileJSScriptingRef;
 typedef TSharedPtr<class FMobileJSScripting> FMobileJSScriptingPtr;
 
@@ -22,7 +24,7 @@ public:
 	static const FString JSMessageTag;
 	static const FString JSMessageHandler;
 
-	FMobileJSScripting(bool bJSBindingToLoweringEnabled, bool bInjectJSOnPageStarted);
+	FMobileJSScripting(bool bJSBindingToLoweringEnabled, bool bInjectJSOnPageStarted, bool bDefaultJSReturnInDict);
 
 	virtual void BindUObject(const FString& Name, UObject* Object, bool bIsPermanent = true) override;
 	virtual void UnbindUObject(const FString& Name, UObject* Object = nullptr, bool bIsPermanent = true) override;
@@ -67,6 +69,7 @@ private:
 	
 	/** When to inject JS code */
 	bool bInjectJSOnPageStarted;
+	bool bDefaultJSReturnInDict;
 };
 
 #endif // PLATFORM_ANDROID  || PLATFORM_IOS  || PLATFORM_MAC
diff --git a/Engine/Source/Runtime/WebBrowser/Private/WebBrowserSingleton.cpp b/Engine/Source/Runtime/WebBrowser/Private/WebBrowserSingleton.cpp
--- a/Engine/Source/Runtime/WebBrowser/Private/WebBrowserSingleton.cpp
+++ b/Engine/Source/Runtime/WebBrowser/Private/WebBrowserSingleton.cpp
@@ -145,7 +145,8 @@ public:
 		bool bInterceptLoadRequests = true,
 		TOptional<FString> ContentsToLoad = TOptional<FString>(),
 		bool ShowErrorMessage = true,
-		FColor BackgroundColor = FColor(255, 255, 255, 255)) override
+		FColor BackgroundColor = FColor(255, 255, 255, 255),
+	    bool bMobileJSReturnInDict = true) override
 	{
 		FCreateBrowserWindowSettings Settings;
 		Settings.OSWindowHandle = OSWindowHandle;
@@ -156,6 +157,7 @@ public:
 		Settings.bShowErrorMessage = ShowErrorMessage;
 		Settings.BackgroundColor = BackgroundColor;
 		Settings.bInterceptLoadRequests = bInterceptLoadRequests;
+		Settings.bMobileJSReturnInDict = bMobileJSReturnInDict;
 
 		return IWebBrowserModule::Get().GetSingleton()->CreateBrowserWindow(Settings);
 	}
@@ -184,7 +186,8 @@ public:
 		bool bInterceptLoadRequests = true,
 		TOptional<FString> ContentsToLoad = TOptional<FString>(),
 		bool ShowErrorMessage = true,
-		FColor BackgroundColor = FColor(255, 255, 255, 255)) override
+		FColor BackgroundColor = FColor(255, 255, 255, 255),
+	    bool bMobileJSReturnInDict = true) override
 	{
 		return nullptr;
 	}
@@ -658,7 +661,8 @@ TSharedPtr<IWebBrowserWindow> FWebBrowserSingleton::CreateBrowserWindow(const FC
 		WindowSettings.bThumbMouseButtonNavigation,
 		WindowSettings.bUseTransparency,
 		bJSBindingsToLoweringEnabled,
-		UserAgentApplication));
+		UserAgentApplication,
+	    WindowSettings.bMobileJSReturnInDict));
 
 	{
 		FScopeLock Lock(&WindowInterfacesCS);
@@ -674,7 +678,8 @@ TSharedPtr<IWebBrowserWindow> FWebBrowserSingleton::CreateBrowserWindow(const FC
 		WindowSettings.bThumbMouseButtonNavigation, 
 		WindowSettings.bUseTransparency,
 		bJSBindingsToLoweringEnabled,
-		UserAgentApplication));
+		UserAgentApplication,
+	    WindowSettings.bMobileJSReturnInDict));
 
 	{
 		FScopeLock Lock(&WindowInterfacesCS);
diff --git a/Engine/Source/Runtime/WebBrowser/Public/IWebBrowserSingleton.h b/Engine/Source/Runtime/WebBrowser/Public/IWebBrowserSingleton.h
--- a/Engine/Source/Runtime/WebBrowser/Public/IWebBrowserSingleton.h
+++ b/Engine/Source/Runtime/WebBrowser/Public/IWebBrowserSingleton.h
@@ -28,7 +28,8 @@ public:
 		bool InterceptLoadRequests = true,
 		TOptional<FString> ContentsToLoad = TOptional<FString>(),
 		bool ShowErrorMessage = true,
-		FColor BackgroundColor = FColor(255, 255, 255, 255)) = 0;
+		FColor BackgroundColor = FColor(255, 255, 255, 255),
+	    bool bMobileJSReturnInDict = true) = 0;
 };
 
 struct FBrowserContextSettings
@@ -67,6 +68,7 @@ struct FCreateBrowserWindowSettings
 		, BrowserFrameRate(24)
 		, Context()
 		, AltRetryDomains()
+		, bMobileJSReturnInDict(true)
 	{ }
 
 	void* OSWindowHandle;
@@ -80,6 +82,7 @@ struct FCreateBrowserWindowSettings
 	int BrowserFrameRate;
 	TOptional<FBrowserContextSettings> Context;
 	TArray<FString> AltRetryDomains;
+	bool bMobileJSReturnInDict;
 };
 
 /**
