Dealing with input issues with WKWebView and Slate on macOS by changing AppleWebBrowserWidget to a SViewport

diff --git a/Engine/Source/Runtime/ApplicationCore/Private/IOS/IOSAppDelegate.cpp b/Engine/Source/Runtime/ApplicationCore/Private/IOS/IOSAppDelegate.cpp
--- a/Engine/Source/Runtime/ApplicationCore/Private/IOS/IOSAppDelegate.cpp
+++ b/Engine/Source/Runtime/ApplicationCore/Private/IOS/IOSAppDelegate.cpp
@@ -1604,7 +1604,7 @@ static FAutoConsoleVariableRef CVarGEnableThermalsReport(
 		self.CommandLineParseTimer = nil;
 	}
 	
-	//    Save openurl infomation before engine initialize.
+	//    Save openurl information before engine initialize.
 	//    When engine is done ready, running like previous. ( if OnOpenUrl is bound on game source. )
 	if (self.bEngineInit)
 	{
@@ -1832,10 +1832,13 @@ extern double GCStartTime;
 	 Save data if appropriate.
 	 See also applicationDidEnterBackground:.
 	 */
-	if (!UE::IsGCLocked())
-	{
-		FCoreDelegates::GetApplicationWillTerminateDelegate().Broadcast();
-	}
+	AsyncTask(ENamedThreads::Type::GameThread, []
+    {
+		if (!UE::IsGCLocked())
+		{
+			FCoreDelegates::GetApplicationWillTerminateDelegate().Broadcast();
+		}
+	});
     
     // note that we are shutting down
     // TODO: fix the reason why we are hanging when asked to shutdown
diff --git a/Engine/Source/Runtime/ApplicationCore/Private/Mac/MacApplication.cpp b/Engine/Source/Runtime/ApplicationCore/Private/Mac/MacApplication.cpp
--- a/Engine/Source/Runtime/ApplicationCore/Private/Mac/MacApplication.cpp
+++ b/Engine/Source/Runtime/ApplicationCore/Private/Mac/MacApplication.cpp
@@ -129,6 +129,7 @@ FMacApplication::FMacApplication()
 ,	AccessibilityCacheTimer(Nil)
 ,	AccessibilityAnnouncementDelayTimer(Nil)
 #endif
+,	LastNativeEvent(nil)
 ,	bDisplayReconfiguring(false)
 {
 	TextInputMethodSystem = MakeShareable(new FMacTextInputMethodSystem);
@@ -745,6 +746,12 @@ void FMacApplication::ProcessEvent(const FDeferredMacEvent& Event)
 	TSharedPtr<FMacWindow> EventWindow = FindWindowByNSWindow(Event.Window);
 	if (Event.Type)
 	{
+		if(LastNativeEvent != nil)
+		{
+			[LastNativeEvent release];
+			LastNativeEvent = nil;
+		}
+		
 		switch (Event.Type)
 		{
 			case NSEventTypeMouseMoved:
@@ -785,11 +792,23 @@ void FMacApplication::ProcessEvent(const FDeferredMacEvent& Event)
 				break;
 
 			case NSEventTypeKeyDown:
+				// retain LastNativeEvent in the event of a ResendLastEvent
+				if(Event.Event != nil)
+				{
+					[Event.Event retain];
+				}
+				LastNativeEvent = Event.Event;
 				ConditionallyUpdateModifierKeys(Event);
 				ProcessKeyDownEvent(Event, EventWindow);
 				break;
 
 			case NSEventTypeKeyUp:
+				// retain LastNativeEvent in the event of a ResendLastEvent
+				if(Event.Event != nil)
+				{
+					[Event.Event retain];
+				}
+				LastNativeEvent = Event.Event;
 				ConditionallyUpdateModifierKeys(Event);
 				ProcessKeyUpEvent(Event);
 				break;
@@ -903,6 +922,14 @@ void FMacApplication::ProcessEvent(const FDeferredMacEvent& Event)
 	}
 }
 
+void FMacApplication::ResendLastEvent() 
+{
+	if(LastNativeEvent != nil) 
+	{
+		ResendEvent(LastNativeEvent);
+	}
+}
+
 void FMacApplication::ResendEvent(NSEvent* Event)
 {
 	MainThreadCall(^{
diff --git a/Engine/Source/Runtime/ApplicationCore/Public/Mac/MacApplication.h b/Engine/Source/Runtime/ApplicationCore/Public/Mac/MacApplication.h
--- a/Engine/Source/Runtime/ApplicationCore/Public/Mac/MacApplication.h
+++ b/Engine/Source/Runtime/ApplicationCore/Public/Mac/MacApplication.h
@@ -200,6 +200,7 @@ public:
 #endif
 
 public:
+	void ResendLastEvent();
 
 	void CloseWindow(TSharedRef<FMacWindow> Window);
 
@@ -421,6 +422,8 @@ private:
 	NSTimer* AccessibilityAnnouncementDelayTimer;
 #endif
 	
+	NSEvent* LastNativeEvent;
+	
 	bool bDisplayReconfiguring;
 
 	friend class FMacWindow;
diff --git a/Engine/Source/Runtime/WebBrowser/Private/Apple/ApplePlatformWebBrowser.cpp b/Engine/Source/Runtime/WebBrowser/Private/Apple/ApplePlatformWebBrowser.cpp
--- a/Engine/Source/Runtime/WebBrowser/Private/Apple/ApplePlatformWebBrowser.cpp
+++ b/Engine/Source/Runtime/WebBrowser/Private/Apple/ApplePlatformWebBrowser.cpp
@@ -17,6 +17,10 @@
 #include "PlatformHttp.h"
 #include "HAL/PlatformProcess.h"
 
+#include "Framework/Application/SlateApplication.h"
+
+#include "Textures/SlateUpdatableTexture.h"
+
 #import <MetalKit/MetalKit.h>
 #include "ExternalTexture.h"
 #include "WebBrowserModule.h"
@@ -37,7 +41,7 @@ TAutoConsoleVariable<bool> CVarWebViewEnableInspectable(
 );
 #endif
 
-class SAppleWebBrowserWidget : public SLeafWidget
+class SAppleWebBrowserWidget : public SViewport
 {
 	SLATE_BEGIN_ARGS(SAppleWebBrowserWidget)
 		: _InitialURL("about:blank")
@@ -67,6 +71,7 @@ class SAppleWebBrowserWidget : public SLeafWidget
 
 		WebBrowserWindowPtr = Args._WebBrowserWindow;
 		Is3DBrowser = false;
+		bAcquiredInitialFocus = false;
 
 		bool bEnableFloatingCloseButton = false;
 		GConfig->GetBool(TEXT("Browser"), TEXT("bEnableFloatingCloseButton"), bEnableFloatingCloseButton, GEngineIni);
@@ -181,6 +186,61 @@ class SAppleWebBrowserWidget : public SLeafWidget
 #endif
 	}
 	
+	void WriteWebBrowserTexture(AppleWebViewWrapper* NativeWebBrowserPtr,
+								FGuid PlayerGuid, FIntPoint Size)
+	{
+		ENQUEUE_RENDER_COMMAND(WriteWebBrowser)(
+			[NativeWebBrowserPtr, PlayerGuid, Size](FRHICommandListImmediate& RHICmdList)
+			{
+				AppleWebViewWrapper* NativeWebBrowser = NativeWebBrowserPtr;
+
+				if (NativeWebBrowser == nil)
+				{
+					return;
+				}
+
+				FTextureRHIRef VideoTexture = [NativeWebBrowser GetVideoTexture];
+				bool bStaleVideoTexture = VideoTexture != nullptr && VideoTexture->GetSizeXY() != Size;
+											
+				if (VideoTexture == nullptr || bStaleVideoTexture)
+				{
+					// We need this texture to be accessible and writable from the CPU because we need to call
+					// replaceRegion:mipmapLevel:withBytes:bytesPerRow: to copy over a screenshot captured in WKWebView
+					// using takeSnapshotWithConfiguration:completionHandler: which gives us an UI/NSImage.
+					const FRHITextureCreateDesc Desc =
+						FRHITextureCreateDesc::Create2D(TEXT("SAppleWebBrowserWidget_VideoTexture"), Size, PF_R8G8B8A8)
+						.SetFlags(ETextureCreateFlags::CPUWritable | ETextureCreateFlags::CPUReadback);
+
+					VideoTexture = RHICmdList.CreateTexture(Desc);
+					[NativeWebBrowser SetVideoTexture : VideoTexture];
+
+					if (VideoTexture == nullptr)
+					{
+						UE_LOG(LogAppleWebBrowser, Warning, TEXT("CreateTexture failed!"));
+						return;
+					}
+
+					[NativeWebBrowser SetVideoTextureValid : false];
+
+				}
+
+				if ([NativeWebBrowser UpdateVideoFrame : VideoTexture->GetNativeResource()])
+				{
+					// if region changed, need to reregister UV scale/offset
+					//UE_LOG(LogAppleWebBrowser, Log, TEXT("UpdateVideoFrame RT: %s"), *Params.PlayerGuid.ToString());
+				}
+
+				if (![NativeWebBrowser IsVideoTextureValid])
+				{
+					FSamplerStateInitializerRHI SamplerStateInitializer(SF_Bilinear, AM_Clamp, AM_Clamp, AM_Clamp);
+					FSamplerStateRHIRef SamplerStateRHI = RHICreateSamplerState(SamplerStateInitializer);
+					FExternalTextureRegistry::Get().RegisterExternalTexture(PlayerGuid, VideoTexture, SamplerStateRHI);
+
+					[NativeWebBrowser SetVideoTextureValid : true];
+				}
+			});
+	}
+	
 	void Tick(const FGeometry& AllottedGeometry, const double InCurrentTime, const float InDeltaTime)
 	{
 		if (WebViewWrapper != nil)
@@ -199,6 +259,15 @@ class SAppleWebBrowserWidget : public SLeafWidget
 				}
 			}
 			
+			// If we are a 3D browser, we should let the Widget Interaction component fully control focus
+			// Otherwise, to keep parity with CEF, we should initially gain focus on the WebBrowser widget.
+			if(!bAcquiredInitialFocus && !Is3DBrowser && InKeyWindow()) 
+			{
+				FSlateApplication::Get().SetKeyboardFocus(SharedThis(this), EFocusCause::SetDirectly);
+				FSlateApplication::Get().ReleaseAllPointerCapture();
+				bAcquiredInitialFocus = true;
+			} 
+			
 			FVector2D Position = AllottedGeometry.GetAccumulatedRenderTransform().GetTranslation();
 			FVector2D Size = TransformVector(AllottedGeometry.GetAccumulatedRenderTransform(), AllottedGeometry.GetLocalSize());
 			FSlateRect SlateRect = AllottedGeometry.GetRenderBoundingRect();
@@ -219,65 +288,40 @@ class SAppleWebBrowserWidget : public SLeafWidget
 
 				if (WebBrowserTexture != nullptr)
 				{
-					struct FWriteWebBrowserParams
-					{
-						AppleWebViewWrapper* NativeWebBrowserPtr;
-						FGuid PlayerGuid;
-						FIntPoint Size;
-					};
 
 					FIntPoint viewportSize = WebBrowserWindowPtr.Pin()->GetViewportSize();
 
-					FWriteWebBrowserParams Params = { WebViewWrapper, WebBrowserTexture->GetExternalTextureGuid(), viewportSize };
-
-					ENQUEUE_RENDER_COMMAND(WriteWebBrowser)(
-						[Params](FRHICommandListImmediate& RHICmdList)
-						{
-							AppleWebViewWrapper* NativeWebBrowser = Params.NativeWebBrowserPtr;
-
-							if (NativeWebBrowser == nil)
-							{
-								return;
-							}
-
-							FTextureRHIRef VideoTexture = [NativeWebBrowser GetVideoTexture];
-							if (VideoTexture == nullptr)
-							{
-								// We need this texture to be accessible and writable from the CPU because we need to call
-								// replaceRegion:mipmapLevel:withBytes:bytesPerRow: to copy over a screenshot captured in WKWebView
-								// using takeSnapshotWithConfiguration:completionHandler: which gives us an UI/NSImage.
-								const FRHITextureCreateDesc Desc =
-									FRHITextureCreateDesc::Create2D(TEXT("SAppleWebBrowserWidget_VideoTexture"), Params.Size, PF_R8G8B8A8)
-									.SetFlags(ETextureCreateFlags::CPUWritable | ETextureCreateFlags::CPUReadback);
-
-								VideoTexture = RHICmdList.CreateTexture(Desc);
-								[NativeWebBrowser SetVideoTexture : VideoTexture];
-
-								if (VideoTexture == nullptr)
-								{
-									UE_LOG(LogAppleWebBrowser, Warning, TEXT("CreateTexture failed!"));
-									return;
-								}
-
-								[NativeWebBrowser SetVideoTextureValid : false];
-
-							}
-
-							if ([NativeWebBrowser UpdateVideoFrame : VideoTexture->GetNativeResource()])
-							{
-								// if region changed, need to reregister UV scale/offset
-								//UE_LOG(LogAppleWebBrowser, Log, TEXT("UpdateVideoFrame RT: %s"), *Params.PlayerGuid.ToString());
-							}
-
-							if (![NativeWebBrowser IsVideoTextureValid])
-							{
-								FSamplerStateInitializerRHI SamplerStateInitializer(SF_Bilinear, AM_Clamp, AM_Clamp, AM_Clamp);
-								FSamplerStateRHIRef SamplerStateRHI = RHICreateSamplerState(SamplerStateInitializer);
-								FExternalTextureRegistry::Get().RegisterExternalTexture(Params.PlayerGuid, VideoTexture, SamplerStateRHI);
-
-								[NativeWebBrowser SetVideoTextureValid : true];
-							}
+					TSharedPtr<FWebBrowserWindow> WebBrowserWindow = WebBrowserWindowPtr.Pin();
+					bool TextureWritten = false;
+										
+					// GetViewportSize not 100% accurate in 3D mode, we use it as a 
+					// reasonable default if we need to only
+					FIntPoint ViewportSize = WebBrowserWindow->GetViewportSize();
+					
+					if(WebBrowserWindowPtr.Pin()->GetParentWindow() != nullptr)
+					{
+						FVector2f WindowSize = WebBrowserWindowPtr.Pin()->GetParentWindow()->GetSizeInScreen();
+						ViewportSize.X = static_cast<int32>(WindowSize.X);
+						ViewportSize.Y = static_cast<int32>(WindowSize.Y);
+#if !PLATFORM_MAC
+						// needs to be on main thread for getScreenScale
+						dispatch_async(dispatch_get_main_queue(), ^ 
+					    {
+							FIntPoint ScaledViewportSize = ViewportSize;
+							CGFloat ScreenScale = [WebViewWrapper getScreenScale];
+							ScaledViewportSize.X *= ScreenScale;
+							ScaledViewportSize.Y *= ScreenScale;
+							
+							WriteWebBrowserTexture(WebViewWrapper, WebBrowserTexture->GetExternalTextureGuid(), ScaledViewportSize);
 						});
+						TextureWritten = true;
+#endif
+					}
+
+					if(!TextureWritten) 
+					{
+						WriteWebBrowserTexture(WebViewWrapper, WebBrowserTexture->GetExternalTextureGuid(), ViewportSize);
+					}
 				}
 			}
 #endif
@@ -569,6 +613,27 @@ class SAppleWebBrowserWidget : public SLeafWidget
 		}
 	}
 	
+	void OnFocus(bool SetFocus)
+	{
+		if(SetFocus) 
+		{
+			FSlateApplication::Get().ReleaseAllPointerCapture();
+		}
+		else if(!FSlateApplication::Get().IsActive() && bFocused) 
+		{
+			// we changed windows but we still want to retain focus and prevent
+			// mouse capture since we were the last widget focused.
+			bAcquiredInitialFocus = false;
+		}
+		
+		bFocused = SetFocus;
+	}
+	
+	bool InKeyWindow()
+	{
+		return WebViewWrapper != nil && [WebViewWrapper inKeyWindow];
+	}
+	
 	~SAppleWebBrowserWidget()
 	{
 		Close();
@@ -582,6 +647,8 @@ private:
 
 	/** Enable 3D appearance */
 	bool Is3DBrowser;
+	bool bAcquiredInitialFocus;
+	bool bFocused;
 
 #if !PLATFORM_TVOS
 	/** The external texture to render the webbrowser output. */
@@ -634,6 +701,8 @@ supportsMetal : (bool)InSupportsMetal supportsMetalMRT : (bool)InSupportsMetalMR
 	bVideoTextureValid = false;
 	bSupportsMetal = InSupportsMetal;
 	bSupportsMetalMRT = InSupportsMetalMRT;
+	
+	BOOL Inspectable = CVarWebViewEnableInspectable.GetValueOnAnyThread()? YES : NO;
 
 #if !PLATFORM_TVOS
 	dispatch_async(dispatch_get_main_queue(), ^
@@ -663,7 +732,7 @@ supportsMetal : (bool)InSupportsMetal supportsMetalMRT : (bool)InSupportsMetalMR
 
 		self.WebView = [[WKWebView alloc]initWithFrame:CGRectMake(1, 1, 100, 100)configuration : theConfiguration];
 #if !UE_BUILD_SHIPPING
-		self.WebView.inspectable = CVarWebViewEnableInspectable.GetValueOnAnyThread()? YES : NO;
+		self.WebView.inspectable = Inspectable;
 #endif
 		[self.WebViewContainer addSubview : WebView];
 		self.WebView.navigationDelegate = self;
@@ -816,6 +885,7 @@ supportsMetal : (bool)InSupportsMetal supportsMetalMRT : (bool)InSupportsMetalMR
 -(void)updateframe:(CGRect)InFrame window:(FGenericWindow*)EnclosedWindow;
 {
 	self.DesiredFrame = InFrame;
+	BOOL Inspectable = CVarWebViewEnableInspectable.GetValueOnAnyThread()? YES : NO;
 
 #if !PLATFORM_TVOS
 	dispatch_async(dispatch_get_main_queue(), ^
@@ -824,7 +894,7 @@ supportsMetal : (bool)InSupportsMetal supportsMetalMRT : (bool)InSupportsMetalMR
 		{
 			
 #if !UE_BUILD_SHIPPING
-			self.WebView.inspectable = CVarWebViewEnableInspectable.GetValueOnAnyThread()? YES : NO;
+			self.WebView.inspectable = Inspectable;
 #endif
 			
 #if PLATFORM_MAC
@@ -838,8 +908,11 @@ supportsMetal : (bool)InSupportsMetal supportsMetalMRT : (bool)InSupportsMetalMR
 				self.WebView.frame = WebViewContainer.bounds;
 			}
 #else
-			self.WebViewContainer.frame = self.DesiredFrame;
-			self.WebView.frame = WebViewContainer.bounds;
+			if(!Is3DBrowser)
+			{
+				self.WebViewContainer.frame = self.DesiredFrame;
+				self.WebView.frame = WebViewContainer.bounds;
+			}
 #endif
 			if (bNeedsAddToView)
 			{
@@ -1086,66 +1159,114 @@ supportsMetal : (bool)InSupportsMetal supportsMetalMRT : (bool)InSupportsMetalMR
 	return true;
 }
 
+#if !PLATFORM_MAC
+-(CGFloat)getScreenScale
+{	
+	for(UIScene* scene in UIApplication.sharedApplication.connectedScenes) 
+	{
+		if([scene isMemberOfClass:[UIWindowScene class]])
+		{
+			for(UIWindow* window in ((UIWindowScene*)scene).windows) 
+			{
+				if(window.isKeyWindow) 
+				{ 
+					return ((UIWindowScene*)scene).screen.scale;
+				}
+			}
+		}
+	}
+	
+	return 1;
+}
+#endif
+
 -(void)updateWebViewMetalTexture:(id<MTLTexture>)texture
 {
 #if !PLATFORM_TVOS
-	@autoreleasepool 
+	
+	dispatch_async(dispatch_get_main_queue(), ^ 
 	{
-		WKSnapshotConfiguration *SnapshotConfig = [[WKSnapshotConfiguration alloc] init];
-		
-		NSUInteger Width = [texture width];
-		NSUInteger Height = [texture height];
-		CGRect SnapshotRect = CGRectMake(0, 0, Width, Height);
-		
-		WebView.frame = SnapshotRect;
-		SnapshotConfig.rect = SnapshotRect;
-		
-		[WebView takeSnapshotWithConfiguration : SnapshotConfig
-							 completionHandler : 
+		@autoreleasepool 
+		{
+			WKSnapshotConfiguration *SnapshotConfig = [[WKSnapshotConfiguration alloc] init];
+			
+			NSUInteger TextureWidth = [texture width];
+			NSUInteger TextureHeight = [texture height];
 #if PLATFORM_MAC
-		 ^ void(NSImage *Image, NSError *Error)
+			NSUInteger Width = TextureWidth;
+			NSUInteger Height = TextureHeight;
+			WebView.frame = CGRectMake(0, 0, Width, Height);
 #else
-		 ^ void(UIImage *Image, NSError *Error)
-#endif
-		{
-			if(Error != nil) 
-			{
-				UE_LOG(LogAppleWebBrowser, Error, TEXT("takeSnapshotWithConfiguration errored with NSError %d!"), Error.code);
-				
-				[SnapshotConfig release];
-				return;
-			}
+			float ScreenScale = [self getScreenScale];
+			NSUInteger Width = TextureWidth / ScreenScale;
+			NSUInteger Height = TextureHeight / ScreenScale;
+			
+			UIView* View = [IOSAppDelegate GetDelegate].IOSView;
+			CGFloat ContentScaleFactor = View.contentScaleFactor;
 			
-			CGColorSpaceRef ColorSpace = CGColorSpaceCreateDeviceRGB();
-			CGContextRef WebViewContext = CGBitmapContextCreate(NULL, Width, Height, 8, 4 * Width, ColorSpace, (CGBitmapInfo)kCGImageAlphaPremultipliedLast);
+			WebView.frame = CGRectMake(0, 0, ContentScaleFactor * Width, ContentScaleFactor * Height);
+#endif
+			SnapshotConfig.rect = CGRectMake(0, 0, WebView.bounds.size.width, WebView.bounds.size.height);
 			
+			[WebView takeSnapshotWithConfiguration : SnapshotConfig
+								 completionHandler : 
 #if PLATFORM_MAC
-			NSGraphicsContext *NewGfxContext = [NSGraphicsContext graphicsContextWithCGContext : WebViewContext
-																  flipped : NO];
-			NSGraphicsContext *OldGfxContext = NSGraphicsContext.currentContext;
-			NSGraphicsContext.currentContext = NewGfxContext;
-			
-			[Image drawInRect: SnapshotRect];
-		 
-			NSGraphicsContext.currentContext = OldGfxContext;
+			 ^ void(NSImage *Image, NSError *Error)
 #else
-			CGContextDrawImage(WebViewContext, SnapshotRect, Image.CGImage);
+			 ^ void(UIImage *Image, NSError *Error)
 #endif
-			[texture replaceRegion : MTLRegionMake2D(0, 0, Width, Height)
-				 mipmapLevel : 0
-				 withBytes : CGBitmapContextGetData(WebViewContext)
-				 bytesPerRow : 4 * Width];
-			
-			Image = nil;
-		 
-			CGColorSpaceRelease(ColorSpace);
-			CGContextRelease(WebViewContext);
-			[SnapshotConfig release];
-		}];
-	}
+			{				
+				if(Error != nil) 
+				{
+					UE_LOG(LogAppleWebBrowser, Error, TEXT("takeSnapshotWithConfiguration errored with NSError %d!"), Error.code);
+					
+					[SnapshotConfig release];
+					return;
+				}
+				
+				CGColorSpaceRef ColorSpace = CGColorSpaceCreateDeviceRGB();
+				CGContextRef WebViewContext = CGBitmapContextCreate(NULL, TextureWidth, TextureHeight, 8, 4 * TextureWidth, ColorSpace, (CGBitmapInfo)kCGImageAlphaPremultipliedLast);
+				
+#if PLATFORM_MAC
+				NSGraphicsContext *NewGfxContext = [NSGraphicsContext graphicsContextWithCGContext : WebViewContext
+																	  flipped : NO];
+				NSGraphicsContext *OldGfxContext = NSGraphicsContext.currentContext;
+				NSGraphicsContext.currentContext = NewGfxContext;
+				
+				[Image drawInRect: WebView.frame];
+			 
+				NSGraphicsContext.currentContext = OldGfxContext;
+#else
+				CGContextDrawImage(WebViewContext, CGRectMake(0, 0, TextureWidth, TextureHeight), Image.CGImage);
+#endif
+				[texture replaceRegion : MTLRegionMake2D(0, 0, TextureWidth, TextureHeight)
+					 mipmapLevel : 0
+					 withBytes : CGBitmapContextGetData(WebViewContext)
+					 bytesPerRow : 4 * TextureWidth];
+				
+				Image = nil;
+			 
+				CGColorSpaceRelease(ColorSpace);
+				CGContextRelease(WebViewContext);
+				[SnapshotConfig release];
+			}];
+		}
+	});
 #endif
 }
 
+-(bool)inKeyWindow;
+{
+	if(self.WebViewContainer != nil && [self.WebViewContainer window] != nil)
+	{
+		return [[self.WebViewContainer window] isKeyWindow];
+	}
+	else 
+	{
+		return NO;
+	}
+}
+
 #if !PLATFORM_TVOS
 
 - (nullable WKWebView *)webView:(WKWebView *)InWebView createWebViewWithConfiguration:(WKWebViewConfiguration *)InConfiguration forNavigationAction:(WKNavigationAction *)InNavigationAction windowFeatures:(WKWindowFeatures *)InWindowFeatures
@@ -1262,14 +1383,20 @@ void FWebBrowserWindow::LoadString(FString Contents, FString DummyURL)
 	BrowserWidget->LoadString(Contents, DummyURL);
 }
 
-TSharedRef<SWidget> FWebBrowserWindow::CreateWidget()
+TOptional<FSlateRenderTransform> FWebBrowserWindow::GetWebBrowserRenderTransform() const
+{
+	return FSlateRenderTransform();
+}
+
+TSharedRef<SViewport> FWebBrowserWindow::CreateWidget()
 {
 	TSharedRef<SAppleWebBrowserWidget> BrowserWidgetRef =
 		SNew(SAppleWebBrowserWidget)
 		.UseTransparency(bUseTransparency)
 		.InitialURL(CurrentUrl)
 		.UserAgentApplication(UserAgentApplication)
-		.WebBrowserWindow(SharedThis(this));
+		.WebBrowserWindow(SharedThis(this))
+		.RenderTransform(this, &FWebBrowserWindow::GetWebBrowserRenderTransform);
 
 	BrowserWidget = BrowserWidgetRef;
 	
@@ -1299,7 +1426,7 @@ FSlateShaderResource* FWebBrowserWindow::GetTexture(bool bIsPopup /*= false*/)
 
 bool FWebBrowserWindow::IsValid() const
 {
-	return false;
+	return true;
 }
 
 bool FWebBrowserWindow::IsInitialized() const
@@ -1329,37 +1456,53 @@ FString FWebBrowserWindow::GetUrl() const
 
 bool FWebBrowserWindow::OnKeyDown(const FKeyEvent& InKeyEvent)
 {
-	return false;
+	RedirectNSEvent();
+	return true;
 }
 
 bool FWebBrowserWindow::OnKeyUp(const FKeyEvent& InKeyEvent)
 {
-	return false;
+	RedirectNSEvent();
+	return true;
 }
 
 bool FWebBrowserWindow::OnKeyChar(const FCharacterEvent& InCharacterEvent)
 {
-	return false;
+	return true;
+}
+
+void FWebBrowserWindow::RedirectNSEvent() 
+{
+	// by default, some events (mainly keystrokes) don't get redirected to NSApp when consumed in Slate
+#if PLATFORM_MAC
+	if (const TSharedPtr<GenericApplication> Application = FSlateApplication::Get().GetPlatformApplication())
+	{
+		if(BrowserWidget != nullptr && BrowserWidget->InKeyWindow())
+		{
+			static_cast<FMacApplication*>(Application.Get())->ResendLastEvent();
+		}
+	}
+#endif
 }
 
 FReply FWebBrowserWindow::OnMouseButtonDown(const FGeometry& MyGeometry, const FPointerEvent& MouseEvent, bool bIsPopup)
 {
-	return FReply::Unhandled();
+	return FReply::Handled().ReleaseMouseCapture().ReleaseMouseLock();
 }
 
 FReply FWebBrowserWindow::OnMouseButtonUp(const FGeometry& MyGeometry, const FPointerEvent& MouseEvent, bool bIsPopup)
 {
-	return FReply::Unhandled();
+	return FReply::Handled().ReleaseMouseCapture().ReleaseMouseLock();
 }
 
 FReply FWebBrowserWindow::OnMouseButtonDoubleClick(const FGeometry& MyGeometry, const FPointerEvent& MouseEvent, bool bIsPopup)
 {
-	return FReply::Unhandled();
+	return FReply::Handled().ReleaseMouseCapture().ReleaseMouseLock();
 }
 
 FReply FWebBrowserWindow::OnMouseMove(const FGeometry& MyGeometry, const FPointerEvent& MouseEvent, bool bIsPopup)
 {
-	return FReply::Unhandled();
+	return FReply::Handled().ReleaseMouseCapture().ReleaseMouseLock();
 }
 
 void FWebBrowserWindow::OnMouseLeave(const FPointerEvent& MouseEvent)
@@ -1378,7 +1521,7 @@ bool FWebBrowserWindow::GetSupportsMouseWheel() const
 
 FReply FWebBrowserWindow::OnMouseWheel(const FGeometry& MyGeometry, const FPointerEvent& MouseEvent, bool bIsPopup)
 {
-	return FReply::Unhandled();
+	return FReply::Handled();
 }
 
 FReply FWebBrowserWindow::OnTouchGesture(const FGeometry& MyGeometry, const FPointerEvent& GestureEvent, bool bIsPopup)
@@ -1389,6 +1532,10 @@ FReply FWebBrowserWindow::OnTouchGesture(const FGeometry& MyGeometry, const FPoi
 
 void FWebBrowserWindow::OnFocus(bool SetFocus, bool bIsPopup)
 {
+	if(BrowserWidget != nullptr)
+	{		
+		BrowserWidget->OnFocus(SetFocus);
+	}
 }
 
 void FWebBrowserWindow::OnCaptureLost()
diff --git a/Engine/Source/Runtime/WebBrowser/Private/Apple/ApplePlatformWebBrowser.h b/Engine/Source/Runtime/WebBrowser/Private/Apple/ApplePlatformWebBrowser.h
--- a/Engine/Source/Runtime/WebBrowser/Private/Apple/ApplePlatformWebBrowser.h
+++ b/Engine/Source/Runtime/WebBrowser/Private/Apple/ApplePlatformWebBrowser.h
@@ -91,6 +91,10 @@ supportsMetal : (bool)InSupportsMetal supportsMetalMRT : (bool)InSupportsMetalMR
 -(bool)UpdateVideoFrame:(void*)ptr;
 -(void)updateWebViewMetalTexture : (id<MTLTexture>)texture;
 -(void)leaveDeadWindow;
+-(bool)inKeyWindow;
+#if !PLATFORM_MAC
+-(CGFloat)getScreenScale;
+#endif
 #if !PLATFORM_TVOS
 -(void)webView:(WKWebView*)InWebView decidePolicyForNavigationAction : (WKNavigationAction*)InNavigationAction decisionHandler : (void(^)(WKNavigationActionPolicy))InDecisionHandler;
 -(void)webView:(WKWebView*)InWebView didCommitNavigation : (WKNavigation*)InNavigation;
@@ -122,7 +126,7 @@ private:
 	/**
 	 * Create the SWidget for this WebBrowserWindow
 	 */
-	TSharedRef<SWidget> CreateWidget();
+	TSharedRef<SViewport> CreateWidget();
 
 public:
 	/** Virtual Destructor. */
@@ -342,6 +346,9 @@ public:
 
 private:
 
+	TOptional<FSlateRenderTransform> GetWebBrowserRenderTransform() const;
+	void RedirectNSEvent();
+	
 	TSharedPtr<SAppleWebBrowserWidget> BrowserWidget;
 
 	/** Current title of this window. */
diff --git a/Engine/Source/Runtime/WebBrowser/Private/SWebBrowserView.cpp b/Engine/Source/Runtime/WebBrowser/Private/SWebBrowserView.cpp
--- a/Engine/Source/Runtime/WebBrowser/Private/SWebBrowserView.cpp
+++ b/Engine/Source/Runtime/WebBrowser/Private/SWebBrowserView.cpp
@@ -1,6 +1,7 @@
 // Copyright Epic Games, Inc. All Rights Reserved.
 
 #include "SWebBrowserView.h"
+#include "Async/Async.h"
 #include "Misc/CommandLine.h"
 #include "Misc/ConfigCacheIni.h"
 #include "Containers/Ticker.h"
@@ -227,7 +228,7 @@ void SWebBrowserView::Construct(const FArguments& InArgs, const TSharedPtr<IWebB
 		}
 
 		BrowserViewport = MakeShareable(new FWebBrowserViewport(BrowserWindow));
-#if WITH_CEF3
+#if WITH_CEF3 || PLATFORM_APPLE
 		BrowserWidget->SetViewportInterface(BrowserViewport.ToSharedRef());
 #endif
 		// If we could not obtain the parent window during widget construction, we'll defer and keep trying.
@@ -496,7 +497,11 @@ void SWebBrowserView::HandleTitleChanged( FString NewTitle )
 void SWebBrowserView::HandleUrlChanged( FString NewUrl )
 {
 	AddressBarUrl = FText::FromString(NewUrl);
-	OnUrlChanged.ExecuteIfBound(AddressBarUrl);
+	
+	AsyncTask(ENamedThreads::Type::GameThread, [this] 
+    {
+		OnUrlChanged.ExecuteIfBound(AddressBarUrl);
+	});
 }
 
 void SWebBrowserView::CloseBrowser()
diff --git a/Engine/Source/Runtime/WebBrowser/Public/SWebBrowserView.h b/Engine/Source/Runtime/WebBrowser/Public/SWebBrowserView.h
--- a/Engine/Source/Runtime/WebBrowser/Public/SWebBrowserView.h
+++ b/Engine/Source/Runtime/WebBrowser/Public/SWebBrowserView.h
@@ -28,7 +28,7 @@ DECLARE_DELEGATE_RetVal_OneParam(TSharedPtr<IToolTip>, FOnCreateToolTip, const F
 DECLARE_DELEGATE_FourParams(FOnConsoleMessageDelegate, const FString& /*Message*/, const FString& /*Source*/, int32 /*Line*/, EWebBrowserConsoleLogSeverity /*Severity*/);
 DECLARE_DELEGATE(FOnFloatingCloseButtonPressedDelegate);
 
-#if WITH_CEF3
+#if WITH_CEF3 || PLATFORM_APPLE
 typedef SViewport SWebBrowserWidget;
 #else
 typedef SWidget SWebBrowserWidget;
