Adding some #if !UE_BUILD_SHIPPING to fix iOS shipping build, and various other if guards to fix support for visionOS and tvOS with AppleWebViewWrapper

diff --git a/Engine/Source/Runtime/WebBrowser/Private/Apple/ApplePlatformWebBrowser.cpp b/Engine/Source/Runtime/WebBrowser/Private/Apple/ApplePlatformWebBrowser.cpp
--- a/Engine/Source/Runtime/WebBrowser/Private/Apple/ApplePlatformWebBrowser.cpp
+++ b/Engine/Source/Runtime/WebBrowser/Private/Apple/ApplePlatformWebBrowser.cpp
@@ -703,7 +703,9 @@ supportsMetal : (bool)InSupportsMetal supportsMetalMRT : (bool)InSupportsMetalMR
 	bSupportsMetal = InSupportsMetal;
 	bSupportsMetalMRT = InSupportsMetalMRT;
 	
+#if !UE_BUILD_SHIPPING
 	BOOL Inspectable = CVarWebViewEnableInspectable.GetValueOnAnyThread()? YES : NO;
+#endif
 
 #if !PLATFORM_TVOS
 	dispatch_async(dispatch_get_main_queue(), ^
@@ -886,7 +888,10 @@ supportsMetal : (bool)InSupportsMetal supportsMetalMRT : (bool)InSupportsMetalMR
 -(void)updateframe:(CGRect)InFrame window:(FGenericWindow*)EnclosedWindow;
 {
 	self.DesiredFrame = InFrame;
+	
+#if !UE_BUILD_SHIPPING
 	BOOL Inspectable = CVarWebViewEnableInspectable.GetValueOnAnyThread()? YES : NO;
+#endif
 
 #if !PLATFORM_TVOS
 	dispatch_async(dispatch_get_main_queue(), ^
@@ -1163,6 +1168,7 @@ supportsMetal : (bool)InSupportsMetal supportsMetalMRT : (bool)InSupportsMetalMR
 #if !PLATFORM_MAC
 -(CGFloat)getScreenScale
 {	
+#if !PLATFORM_VISIONOS
 	for(UIScene* scene in UIApplication.sharedApplication.connectedScenes) 
 	{
 		if([scene isMemberOfClass:[UIWindowScene class]])
@@ -1176,6 +1182,7 @@ supportsMetal : (bool)InSupportsMetal supportsMetalMRT : (bool)InSupportsMetalMR
 			}
 		}
 	}
+#endif
 	
 	return 1;
 }
@@ -1258,6 +1265,9 @@ supportsMetal : (bool)InSupportsMetal supportsMetalMRT : (bool)InSupportsMetalMR
 
 -(bool)inKeyWindow;
 {
+#if PLATFORM_TVOS
+	return YES;
+#else
 	if(self.WebViewContainer != nil && [self.WebViewContainer window] != nil)
 	{
 		return [[self.WebViewContainer window] isKeyWindow];
@@ -1266,6 +1276,7 @@ supportsMetal : (bool)InSupportsMetal supportsMetalMRT : (bool)InSupportsMetalMR
 	{
 		return NO;
 	}
+#endif
 }
 
 #if !PLATFORM_TVOS
