Fix the WebBrowser freeze when bring the game to background while loading a URL.

diff --git a/Engine/Source/Runtime/Launch/Private/Android/AndroidEventManager.cpp b/Engine/Source/Runtime/Launch/Private/Android/AndroidEventManager.cpp
--- a/Engine/Source/Runtime/Launch/Private/Android/AndroidEventManager.cpp
+++ b/Engine/Source/Runtime/Launch/Private/Android/AndroidEventManager.cpp
@@ -66,6 +66,46 @@ static const TCHAR* GetAppEventName(EAppEventState State)
 	}
 }
 
+void FAppEventManager::SetAppIsActive(bool bIsActive)
+{
+	bAppIsActive = bIsActive;
+}
+
+void FAppEventManager::ResumePausedGameThread()
+{
+	if (bAppIsActive || bTempRunningTriggered)
+	{
+		//UE_LOG(LogAndroid, Log, TEXT("Resume_GT_Temp bypassed. %d, %d, %d"), (bool)bAppIsActive, (bool)bRunning, (bool)bTempRunningTriggered);
+		return;
+	}
+
+	EnqueueAppEvent(APP_EVENT_STATE_APP_ACTIVATED);
+
+	if (EventHandlerEvent)
+	{
+		EventHandlerEvent->Trigger();
+	}
+
+	bTempRunningTriggered = true;
+
+	//UE_LOG(LogAndroid, Log, TEXT("Resume_GT_Temp done. %d, %d, %d"), (bool)bAppIsActive, (bool)bRunning, (bool)bTempRunningTriggered);
+}
+
+void FAppEventManager::KeepGameThreadPaused()
+{
+	if (bRunning && !bAppIsActive && bTempRunningTriggered)
+	{
+		EnqueueAppEvent(APP_EVENT_STATE_APP_SUSPENDED);
+		//UE_LOG(LogAndroid, Log, TEXT("Pause_GT_Temp done. %d, %d, %d"), (bool)bAppIsActive, (bool)bRunning, (bool)bTempRunningTriggered);
+	}
+	else
+	{
+		//UE_LOG(LogAndroid, Log, TEXT("Pause_GT_Temp bypassed. %d, %d, %d"), (bool)bAppIsActive, (bool)bRunning, (bool)bTempRunningTriggered);
+	}
+
+	bTempRunningTriggered = false;
+}
+
 
 void FAppEventManager::Tick()
 {
@@ -207,6 +247,8 @@ FAppEventManager::FAppEventManager():
 	,bHaveWindow(false)
 	,bHaveGame(false)
 	,bRunning(false)
+	,bTempRunningTriggered(false)
+	,bAppIsActive(false)
 {
 	pthread_mutex_init(&QueueMutex, NULL);
 
diff --git a/Engine/Source/Runtime/Launch/Private/Android/LaunchAndroid.cpp b/Engine/Source/Runtime/Launch/Private/Android/LaunchAndroid.cpp
--- a/Engine/Source/Runtime/Launch/Private/Android/LaunchAndroid.cpp
+++ b/Engine/Source/Runtime/Launch/Private/Android/LaunchAndroid.cpp
@@ -1473,6 +1473,8 @@ static void ActivateApp_EventThread()
 	FAppEventManager::GetInstance()->EnqueueAppEvent(APP_EVENT_STATE_APP_ACTIVATED);
 
 	bAppIsActive_EventThread = true;
+	FAppEventManager::GetInstance()->SetAppIsActive(bAppIsActive_EventThread);
+
 	UE::GC::Private::SetSkipDestroy(false);
 	STANDALONE_DEBUG_LOGf(LogAndroid, TEXT("event thread, activate app, UNLOCKED HW window lock. bReadyToProcessEvents=%d, bAppIsActive_EventThread=%d"), bReadyToProcessEvents, bAppIsActive_EventThread);
 	FAppEventManager::GetInstance()->EnqueueAppEvent(APP_EVENT_RUN_CALLBACK, FAppEventData([]()
@@ -1558,6 +1560,8 @@ static void SuspendApp_EventThread()
 		return;
 	}
 	bAppIsActive_EventThread = false;
+	FAppEventManager::GetInstance()->SetAppIsActive(bAppIsActive_EventThread);
+
 	UE::GC::Private::SetSkipDestroy(true);
 	// Lock the window, this prevents event thread from removing the window whilst the RHI initializes.
 
diff --git a/Engine/Source/Runtime/Launch/Public/Android/AndroidEventManager.h b/Engine/Source/Runtime/Launch/Public/Android/AndroidEventManager.h
--- a/Engine/Source/Runtime/Launch/Public/Android/AndroidEventManager.h
+++ b/Engine/Source/Runtime/Launch/Public/Android/AndroidEventManager.h
@@ -4,6 +4,7 @@
 
 #include "CoreMinimal.h"
 #include "Containers/Queue.h"
+#include "Templates/Atomic.h"
 
 #if USE_ANDROID_EVENTS
 
@@ -78,13 +79,40 @@ public:
 	void SetEmptyQueueHandlerEvent(FEvent* InEventHandlerEvent);
 	void TriggerEmptyQueue();
 
+	void SetAppIsActive(bool bIsActive);
+
 	void PauseAudio();
 	void ResumeAudio();
 	static void ReleaseMicrophone(bool shuttingDown);
 
+	// When the game thread is stalled by operations like switching to background,
+	// GameThreadTicker can be used to resume the game thread for 1 frame and 
+	// stall it again in the subsequent frame.
+	// If a wait call is behind a GameThreadTicker object on another thread, the GT will
+	// be resumed until the wait returns and the life time of the object ends.
+	class GameThreadTicker
+	{
+	public:
+
+		GameThreadTicker()
+		{
+			FAppEventManager::GetInstance()->ResumePausedGameThread();
+		}
+
+		~GameThreadTicker()
+		{
+			FAppEventManager::GetInstance()->KeepGameThreadPaused();
+		}
+	};
+
+	friend class GameThreadTicker;
+
 protected:
 	FAppEventManager();
 
+	void ResumePausedGameThread();
+	void KeepGameThreadPaused();
+
 private:
 	FAppEventPacket DequeueAppEvent();
 
@@ -111,7 +139,10 @@ private:
 
 	bool bHaveWindow;
 	bool bHaveGame;
-	bool bRunning;
+	TAtomic<bool> bRunning;
+	
+	TAtomic<bool> bTempRunningTriggered;// Flag to identify if the temporary running has been trigged so that APP_EVENT_STATE_APP_ACTIVATED will be only enqueued once.
+	TAtomic<bool> bAppIsActive;			// Synced with bAppIsActive_EventThread in LaunchAndroid
 };
 
 bool IsInAndroidEventThread();
diff --git a/Engine/Source/Runtime/WebBrowser/Private/Android/AndroidWebBrowserWidget.cpp b/Engine/Source/Runtime/WebBrowser/Private/Android/AndroidWebBrowserWidget.cpp
--- a/Engine/Source/Runtime/WebBrowser/Private/Android/AndroidWebBrowserWidget.cpp
+++ b/Engine/Source/Runtime/WebBrowser/Private/Android/AndroidWebBrowserWidget.cpp
@@ -11,6 +11,7 @@
 #include "Android/AndroidApplication.h"
 #include "Android/AndroidWindow.h"
 #include "Android/AndroidJava.h"
+#include "Android/AndroidEventManager.h"
 #include "Async/Async.h"
 #include "Misc/ScopeLock.h"
 #include "RHICommandList.h"
@@ -715,6 +716,8 @@ jbyteArray SAndroidWebBrowserWidget::HandleShouldInterceptRequest(jstring JUrl)
 				}
 			}
 		}, TStatId(), NULL, ENamedThreads::GameThread);
+
+		FAppEventManager::GameThreadTicker Ticker;
 		FTaskGraphInterface::Get().WaitUntilTaskCompletes(OnLoadUrl);
 	}
 
@@ -752,6 +755,8 @@ bool SAndroidWebBrowserWidget::HandleShouldOverrideUrlLoading(jstring JUrl)
 			}
 		}
 	}, TStatId(), NULL, ENamedThreads::GameThread);
+
+	FAppEventManager::GameThreadTicker Ticker;
 	FTaskGraphInterface::Get().WaitUntilTaskCompletes(OnBeforeBrowse);
 
 	return Retval;
@@ -789,6 +794,8 @@ bool SAndroidWebBrowserWidget::HandleJsDialog(TSharedPtr<IWebBrowserDialog>& Dia
 			}
 		}
 	}, TStatId(), NULL, ENamedThreads::GameThread);
+
+	FAppEventManager::GameThreadTicker Ticker;
 	FTaskGraphInterface::Get().WaitUntilTaskCompletes(OnShowDialog);
 
 	return Retval;
@@ -809,6 +816,8 @@ void SAndroidWebBrowserWidget::HandleReceivedTitle(jstring JTitle)
 			{
 				BrowserWindow->SetTitle(Title);
 			}, TStatId(), NULL, ENamedThreads::GameThread);
+
+			FAppEventManager::GameThreadTicker Ticker;
 			FTaskGraphInterface::Get().WaitUntilTaskCompletes(OnSetTitle);
 		}
 	}
@@ -831,6 +840,8 @@ void SAndroidWebBrowserWidget::HandlePageLoad(jstring JUrl, bool bIsLoading, int
 			{
 				BrowserWindow->NotifyDocumentLoadingStateChange(Url, bIsLoading);
 			}, TStatId(), NULL, ENamedThreads::GameThread);
+
+			FAppEventManager::GameThreadTicker Ticker;
 			FTaskGraphInterface::Get().WaitUntilTaskCompletes(OnNotifyDocumentLoadingStateChange);
 		}
 	}
@@ -850,6 +861,8 @@ void SAndroidWebBrowserWidget::HandleReceivedError(jint ErrorCode, jstring /* ig
 			{
 				BrowserWindow->NotifyDocumentError(Url, ErrorCode);
 			}, TStatId(), NULL, ENamedThreads::GameThread);
+
+			FAppEventManager::GameThreadTicker Ticker;
 			FTaskGraphInterface::Get().WaitUntilTaskCompletes(OnNotifyDocumentError);
 		}
 	}
@@ -866,6 +879,8 @@ void SAndroidWebBrowserWidget::FloatingCloseButtonPressed()
 				{
 					BrowserWindow->FloatingCloseButtonPressed();
 				}, TStatId(), NULL, ENamedThreads::GameThread);
+
+			FAppEventManager::GameThreadTicker Ticker;
 			FTaskGraphInterface::Get().WaitUntilTaskCompletes(OnRequestClose);
 		}
 	}
