3D UI Input on WKWebView WebBrowser backend through manual JS event dispatching

diff --git a/Engine/Source/Runtime/InputCore/Private/Mac/MacPlatformInput.cpp b/Engine/Source/Runtime/InputCore/Private/Mac/MacPlatformInput.cpp
--- a/Engine/Source/Runtime/InputCore/Private/Mac/MacPlatformInput.cpp
+++ b/Engine/Source/Runtime/InputCore/Private/Mac/MacPlatformInput.cpp
@@ -93,4 +93,30 @@ FKey FMacPlatformInput::GetPlatformDeleteKey()
 	}
 
 	return bPlatformDeleteIsBackspace ? EKeys::BackSpace : EKeys::Delete;
-}
\ No newline at end of file
+}
+
+uint32 FMacPlatformInput::GetMacNativeModifierKeyCode(uint32 ModifierKeyCode)
+{
+	switch(ModifierKeyCode) 
+	{
+		case MMK_RightCommand:
+		case MMK_LeftCommand:
+			return kVK_Command;
+		case MMK_LeftShift:
+			return kVK_Shift;
+		case MMK_CapsLock:
+			return kVK_CapsLock;
+		case MMK_LeftAlt:
+			return kVK_Option;
+		case MMK_LeftControl:
+			return kVK_Control;
+		case MMK_RightShift:
+			return kVK_RightShift;
+		case MMK_RightAlt:
+			return kVK_RightOption;
+		case MMK_RightControl:
+			return kVK_RightControl;
+		default:
+			return ModifierKeyCode;
+	}
+}
diff --git a/Engine/Source/Runtime/InputCore/Public/Mac/MacPlatformInput.h b/Engine/Source/Runtime/InputCore/Public/Mac/MacPlatformInput.h
--- a/Engine/Source/Runtime/InputCore/Public/Mac/MacPlatformInput.h
+++ b/Engine/Source/Runtime/InputCore/Public/Mac/MacPlatformInput.h
@@ -10,6 +10,7 @@ struct INPUTCORE_API FMacPlatformInput : FGenericPlatformInput
 	static uint32 GetKeyMap( uint32* KeyCodes, FString* KeyNames, uint32 MaxMappings );
 	static uint32 GetCharKeyMap(uint32* KeyCodes, FString* KeyNames, uint32 MaxMappings);
 	static FKey GetPlatformDeleteKey();
+	static uint32 GetMacNativeModifierKeyCode(uint32 ModifierKeyCode);
 };
 
 typedef FMacPlatformInput FPlatformInput;
diff --git a/Engine/Source/Runtime/WebBrowser/Private/Apple/ApplePlatformWebBrowser.cpp b/Engine/Source/Runtime/WebBrowser/Private/Apple/ApplePlatformWebBrowser.cpp
--- a/Engine/Source/Runtime/WebBrowser/Private/Apple/ApplePlatformWebBrowser.cpp
+++ b/Engine/Source/Runtime/WebBrowser/Private/Apple/ApplePlatformWebBrowser.cpp
@@ -9,16 +9,18 @@
 #include "IOS/IOSAppDelegate.h"
 #else
 #include "Mac/CocoaWindow.h"
+#include "Mac/CocoaThread.h"
 #include "Mac/MacWindow.h"
 #include "Mac/MacApplication.h"
 #endif
 
+#include "Async/Async.h"
 #include "Widgets/SLeafWidget.h"
 #include "PlatformHttp.h"
 #include "HAL/PlatformProcess.h"
 
 #include "Framework/Application/SlateApplication.h"
-
+#include "Rendering/SlateRenderer.h"
 #include "Textures/SlateUpdatableTexture.h"
 
 #import <MetalKit/MetalKit.h>
@@ -29,6 +31,7 @@
 
 #include "Apple/AppleAsyncTask.h"
 #include "AppleCookieManager.h"
+#include "WebView3DInputHandler.h"
 
 DEFINE_LOG_CATEGORY(LogAppleWebBrowser);
 
@@ -134,6 +137,8 @@ class SAppleWebBrowserWidget : public SViewport
 		CGRect NewFrame;
 #if PLATFORM_MAC
 		TSharedPtr<FGenericWindow> EnclosedWindow = WebBrowserWindowPtr.Pin()->GetParentWindow().Get()->GetNativeWindow();
+		
+		if(!Is3DBrowser) 
 		{
 			NSWindow* ParentWindow = (NSWindow*)EnclosedWindow->GetOSWindowHandle();
 			float ScaleFactor = [ParentWindow backingScaleFactor];
@@ -290,8 +295,7 @@ class SAppleWebBrowserWidget : public SViewport
 				if (WebBrowserTexture != nullptr)
 				{
 
-					FIntPoint viewportSize = WebBrowserWindowPtr.Pin()->GetViewportSize();
-
+					
 					TSharedPtr<FWebBrowserWindow> WebBrowserWindow = WebBrowserWindowPtr.Pin();
 					bool TextureWritten = false;
 										
@@ -616,7 +620,7 @@ class SAppleWebBrowserWidget : public SViewport
 	
 	void OnFocus(bool SetFocus)
 	{
-		if(SetFocus) 
+		if(SetFocus && !Is3D()) 
 		{
 			FSlateApplication::Get().ReleaseAllPointerCapture();
 		}
@@ -630,6 +634,11 @@ class SAppleWebBrowserWidget : public SViewport
 		bFocused = SetFocus;
 	}
 	
+	bool Is3D() const
+	{
+		return Is3DBrowser;
+	}
+	
 	bool InKeyWindow()
 	{
 		return WebViewWrapper != nil && [WebViewWrapper inKeyWindow];
@@ -689,6 +698,9 @@ private:
 #endif
 @synthesize NextURL;
 @synthesize NextContent;
+#if PLATFORM_MAC
+@synthesize GameWindowView;
+#endif
 
 -(void)create:(SAppleWebBrowserWidget*)InWebBrowserWidget userAgentApplication: (NSString*)UserAgentApplication useTransparency : (bool)InUseTransparency
 supportsMetal : (bool)InSupportsMetal supportsMetalMRT : (bool)InSupportsMetalMRT enableFloatingCloseButton : (bool)bEnableFloatingCloseButton;
@@ -706,6 +718,9 @@ supportsMetal : (bool)InSupportsMetal supportsMetalMRT : (bool)InSupportsMetalMR
 #if !UE_BUILD_SHIPPING
 	BOOL Inspectable = CVarWebViewEnableInspectable.GetValueOnAnyThread()? YES : NO;
 #endif
+#if PLATFORM_MAC
+	GameWindowView = nil;
+#endif
 
 #if !PLATFORM_TVOS
 	dispatch_async(dispatch_get_main_queue(), ^
@@ -906,7 +921,39 @@ supportsMetal : (bool)InSupportsMetal supportsMetalMRT : (bool)InSupportsMetalMR
 #if PLATFORM_MAC
 			if(Is3DBrowser) 
 			{
-				bNeedsAddToView = false;
+				if(bNeedsAddToView || !Current3DWindow.IsValid())
+				{
+					// we handle display of 3D WebViews but choose to still add it to the view hierarchy
+					// as a completely transparent view for better integration 
+					// (allows audio playback, some tick-like JS events dispatched for us, probably some other stuff)
+					self.WebViewContainer.alphaValue = 0.0;
+					bNeedsAddToView = false;
+					
+					AsyncTask(ENamedThreads::Type::GameThread, [self]() 
+				    {
+						if (const TSharedPtr<SWindow> ActiveWindow = FSlateApplication::Get().GetActiveTopLevelWindow())
+						{
+							FMacWindow* MacGameWindow = static_cast<FMacWindow*>(ActiveWindow->GetNativeWindow().Get());
+							
+							if(MacGameWindow != nullptr)
+							{
+								FCocoaWindow* GameCocoaWindow = MacGameWindow->GetWindowHandle();
+								
+								Current3DWindow = ActiveWindow;
+								self.GameWindowView = GameCocoaWindow.openGLView;
+							}
+						}
+						else
+						{
+							bNeedsAddToView = true;
+						}
+					});
+				}
+				else if(GameWindowView != nil)
+				{
+					[self.GameWindowView addSubview : WebViewContainer];
+					self.GameWindowView = nil;
+				}
 			}
 			else
 			{
@@ -1000,8 +1047,8 @@ supportsMetal : (bool)InSupportsMetal supportsMetalMRT : (bool)InSupportsMetalMR
 			if(Error != nil) 
 			{
 				const char *domainString = [Error.domain UTF8String];
-				UE_LOG(LogAppleWebBrowser, Error, TEXT("executeJavaScript(\"%s\") errored with NSError %d in domain %hs!"), 
-					   *InJavaScriptFString, Error.code, domainString);
+				/*UE_LOG(LogAppleWebBrowser, Error, TEXT("executeJavaScript(\"%s\") errored with NSError %d in domain %hs!"), 
+					   *InJavaScriptFString, Error.code, domainString);*/
 			}
 		}];
 	});
@@ -1378,6 +1425,7 @@ FWebBrowserWindow::FWebBrowserWindow(FString InUrl, TOptional<FString> InContent
 	GConfig->GetBool(TEXT("Browser"), TEXT("bInjectJSOnPageStarted"), bInjectJSOnPageStarted, GEngineIni);
 	
 	Scripting = MakeShared<FMobileJSScripting>(bInJSBindingToLoweringEnabled, bInjectJSOnPageStarted, bMobileJSReturnInDict);
+	WebView3DInput = MakeShared<FWebView3DInputHandler>();
 }
 
 FWebBrowserWindow::~FWebBrowserWindow()
@@ -1451,6 +1499,11 @@ bool FWebBrowserWindow::IsClosing() const
 	return false;
 }
 
+bool FWebBrowserWindow::Is3D() 
+{
+	return BrowserWidget != nil ? BrowserWidget->Is3D() : false;
+}
+
 EWebBrowserDocumentState FWebBrowserWindow::GetDocumentLoadingState() const
 {
 	return DocumentState;
@@ -1468,18 +1521,36 @@ FString FWebBrowserWindow::GetUrl() const
 
 bool FWebBrowserWindow::OnKeyDown(const FKeyEvent& InKeyEvent)
 {
-	RedirectNSEvent();
+	if(Is3D())
+	{
+		WebView3DInput->SendKeyEventToJS(*this, "keydown", InKeyEvent);
+	}
+	else
+	{
+		RedirectNSEvent();
+	}
 	return true;
 }
 
 bool FWebBrowserWindow::OnKeyUp(const FKeyEvent& InKeyEvent)
 {
-	RedirectNSEvent();
+	if(Is3D())
+	{
+		WebView3DInput->SendKeyEventToJS(*this, "keyup", InKeyEvent);
+	}
+	else
+	{
+		RedirectNSEvent();
+	}
 	return true;
 }
 
 bool FWebBrowserWindow::OnKeyChar(const FCharacterEvent& InCharacterEvent)
 {
+	if(Is3D()) 
+	{
+		WebView3DInput->SendCharEventToJS(*this, InCharacterEvent);
+	}
 	return true;
 }
 
@@ -1499,21 +1570,42 @@ void FWebBrowserWindow::RedirectNSEvent()
 
 FReply FWebBrowserWindow::OnMouseButtonDown(const FGeometry& MyGeometry, const FPointerEvent& MouseEvent, bool bIsPopup)
 {
+	if(Is3D()) 
+	{		
+		WebView3DInput->SendMouseEventToJS(*this, "down", MouseEvent);
+	}
+	
 	return FReply::Handled().ReleaseMouseCapture().ReleaseMouseLock();
 }
 
 FReply FWebBrowserWindow::OnMouseButtonUp(const FGeometry& MyGeometry, const FPointerEvent& MouseEvent, bool bIsPopup)
 {
+	if(Is3D()) 
+	{
+		WebView3DInput->SendMouseEventToJS(*this, "up", MouseEvent);
+		WebView3DInput->SendMouseEventToJS(*this, "click", MouseEvent);
+	}
+
 	return FReply::Handled().ReleaseMouseCapture().ReleaseMouseLock();
 }
 
 FReply FWebBrowserWindow::OnMouseButtonDoubleClick(const FGeometry& MyGeometry, const FPointerEvent& MouseEvent, bool bIsPopup)
 {
+	if(Is3D())
+	{
+		WebView3DInput->SendMouseEventToJS(*this, "dblclick", MouseEvent);
+	}
+	
 	return FReply::Handled().ReleaseMouseCapture().ReleaseMouseLock();
 }
 
 FReply FWebBrowserWindow::OnMouseMove(const FGeometry& MyGeometry, const FPointerEvent& MouseEvent, bool bIsPopup)
 {
+	if(Is3D()) 
+	{
+		WebView3DInput->SendMouseEventToJS(*this, "move", MouseEvent);
+	}
+	
 	return FReply::Handled().ReleaseMouseCapture().ReleaseMouseLock();
 }
 
@@ -1533,6 +1625,11 @@ bool FWebBrowserWindow::GetSupportsMouseWheel() const
 
 FReply FWebBrowserWindow::OnMouseWheel(const FGeometry& MyGeometry, const FPointerEvent& MouseEvent, bool bIsPopup)
 {
+	if(Is3D())
+	{
+		WebView3DInput->SendMouseEventToJS(*this, "wheel", MouseEvent);
+	}
+	
 	return FReply::Handled();
 }
 
diff --git a/Engine/Source/Runtime/WebBrowser/Private/Apple/ApplePlatformWebBrowser.h b/Engine/Source/Runtime/WebBrowser/Private/Apple/ApplePlatformWebBrowser.h
--- a/Engine/Source/Runtime/WebBrowser/Private/Apple/ApplePlatformWebBrowser.h
+++ b/Engine/Source/Runtime/WebBrowser/Private/Apple/ApplePlatformWebBrowser.h
@@ -45,6 +45,7 @@ class SWebBrowserView;
 {
 	SAppleWebBrowserWidget* WebBrowserWidget;
 	FTextureRHIRef VideoTexture;
+	TWeakPtr<SWindow> Current3DWindow;
 	bool bNeedsAddToView;
 	bool Is3DBrowser;
 	bool bVideoTextureValid;
@@ -55,6 +56,7 @@ class SWebBrowserView;
 @property(strong) WKWebView* WebView;
 #if PLATFORM_MAC
 @property(strong) NSView* WebViewContainer;
+@property(strong) NSView* GameWindowView;
 #else
 @property(strong) WebViewCloseButton* CloseButton;
 @property(strong) UIView* WebViewContainer;
@@ -101,6 +103,8 @@ supportsMetal : (bool)InSupportsMetal supportsMetalMRT : (bool)InSupportsMetalMR
 #endif
 @end
 
+class FWebView3DInputHandler;
+
 /**
 * Implementation of interface for dealing with a Web Browser window.
 */
@@ -331,9 +335,11 @@ public:
 	* Browser's visibility
 	*/
 	bool IsVisible();
+	
+	bool Is3D();
 
 	/**
-	* Webvew floating close button was pressed
+	* Webview floating close button was pressed
 	*/
 	void FloatingCloseButtonPressed();
 
@@ -433,6 +439,7 @@ private:
 	int ErrorCode;
 
 	FMobileJSScriptingPtr Scripting;
+	TSharedPtr<FWebView3DInputHandler> WebView3DInput;
 
 	mutable TOptional<TFunction<void(const FString&)>> GetPageSourceCallback;
 	/** Delegate for handling key down events not handled by the browser. */
diff --git a/Engine/Source/Runtime/WebBrowser/Private/Apple/WebView3DInputHandler.cpp b/Engine/Source/Runtime/WebBrowser/Private/Apple/WebView3DInputHandler.cpp
new file mode 100644
--- /dev/null
+++ b/Engine/Source/Runtime/WebBrowser/Private/Apple/WebView3DInputHandler.cpp
@@ -0,0 +1,250 @@
+// Copyright Epic Games, Inc. All Rights Reserved.
+
+#include "WebView3DInputHandler.h"
+
+#include "IWebBrowserWindow.h"
+
+#if PLATFORM_MAC
+#include "HAL/PlatformInput.h"
+#endif
+
+DEFINE_LOG_CATEGORY(LogWebView3DInputHandler);
+
+FWebView3DInputHandler::FWebView3DInputHandler()
+	: bCtrlKeyDown(false) 
+	, bAltKeyDown(false)
+	, bCmdKeyDown(false) 
+	, bShiftKeyDown(false)
+{
+}
+
+FString FWebView3DInputHandler::BoolToString(bool InBool)
+{
+	return InBool ? "true" : "false";
+}
+
+FString FWebView3DInputHandler::GetModifierKeyStatus()
+{
+	return FString::Printf(TEXT("altKey: %s, ctrlKey: %s, metaKey: %s, shiftKey: %s"), 
+						   *BoolToString(bAltKeyDown), *BoolToString(bCtrlKeyDown), 
+						   *BoolToString(bCmdKeyDown), *BoolToString(bShiftKeyDown));
+}
+
+bool FWebView3DInputHandler::CheckIfCharacterIsInputable(TCHAR CharToInput)
+{
+	// is there a more UE way to do this?
+	return ![NSCharacterSet.illegalCharacterSet characterIsMember: CharToInput] && 
+		   ![NSCharacterSet.controlCharacterSet characterIsMember: CharToInput];
+}
+
+FString FWebView3DInputHandler::GetCurrentJSElement()
+{
+	return FString::Printf(TEXT("document.elementFromPoint(%f, %f)"), MousePos.X, MousePos.Y); 
+}
+
+void FWebView3DInputHandler::SendMouseEventToJS(IWebBrowserWindow& InWindow, const FString& EventName, const FPointerEvent& MouseEvent) 
+{
+	MousePos = MouseEvent.GetScreenSpacePosition();
+	TSet<FKey> PressedButtons = MouseEvent.GetPressedButtons();
+	FKey EffectingButton = MouseEvent.GetEffectingButton();
+	int32_t PressedButtonSum = 0;
+	FString DispatchEventJS;
+	
+	if(EffectingButton != EKeys::Invalid) 
+	{
+		PressedButtons.Add(EffectingButton);
+	}
+	
+	for(FKey Button : PressedButtons) 
+	{
+		if(Button == EKeys::LeftMouseButton) 
+		{
+			PressedButtonSum += 1;
+		} 
+		else if(Button == EKeys::RightMouseButton) 
+		{
+			PressedButtonSum += 2;
+		} 
+		else if(Button == EKeys::MiddleMouseButton) 
+		{
+			PressedButtonSum += 4;
+		}
+	}
+	
+	FString CurrentJSElement = FWebView3DInputHandler::GetCurrentJSElement();
+	FString BaseDispatch = FString::Printf(TEXT("%s.dispatchEvent"), *CurrentJSElement);
+	FString BaseMouseEventData = FString::Printf(TEXT("view: window, clientX: %f, clientY: %f, buttons: %d, bubbles: true, %s"), 
+												 MousePos.X, MousePos.Y, PressedButtonSum, *GetModifierKeyStatus());
+	
+	FString MouseEventName = EventName;
+	FString PointerEventName;
+	
+	if(EventName == "wheel")
+	{
+		// we need to update scrollLeft and scrollTop of the target, dispatch the wheel event and the scroll event
+		float WheelDelta = MouseEvent.GetWheelDelta();
+		float ScrollDelta = InWindow.GetViewportSize().Y * WheelDelta;
+		
+		FString UpdateScrollJS = FString::Printf(TEXT("window.scrollBy(0.0, %f)"), ScrollDelta);
+		InWindow.ExecuteJavascript(UpdateScrollJS);
+		
+		DispatchEventJS = FString::Printf(TEXT("%s(new WheelEvent(\"%s\", {%s, deltaY: %f}))"), *BaseDispatch, *EventName, *BaseMouseEventData, ScrollDelta);
+	}
+	else 
+	{
+		// Most sites should still work with just mouse events, but we dispatch
+		// pointer events as well for newer sites which may take them for granted.
+		if(EventName != "click" && EventName != "dblclick")
+		{
+			MouseEventName = FString::Printf(TEXT("mouse%s"), *EventName);
+			PointerEventName = FString::Printf(TEXT("pointer%s"), *EventName);
+		}
+		
+		if(EffectingButton != EKeys::Invalid) 
+		{
+			int32_t PressedButton = 0;
+			
+			if(EffectingButton == EKeys::LeftMouseButton) 
+			{
+				PressedButton = 0;
+			} 
+			else if(EffectingButton == EKeys::RightMouseButton) 
+			{
+				PressedButton = 2;
+			} 
+			else if(EffectingButton == EKeys::MiddleMouseButton) 
+			{
+				PressedButton = 1;
+			}
+			
+			// follows JS convention for mouse buttons: 0 (left), 1 (middle), 2 (right)
+			DispatchEventJS = FString::Printf(TEXT("%s(new MouseEvent(\"%s\", {%s, button: %d}))"), *BaseDispatch, *MouseEventName, 
+											  *BaseMouseEventData, PressedButton);
+			
+			if(PointerEventName.Len() > 0)
+			{
+				InWindow.ExecuteJavascript(DispatchEventJS);
+				
+				DispatchEventJS = FString::Printf(TEXT("%s(new PointerEvent(\"%s\", {%s, button: %d}))"), *BaseDispatch, *PointerEventName, 
+												  *BaseMouseEventData, PressedButton);
+			}
+		}
+		else
+		{
+			DispatchEventJS = FString::Printf(TEXT("%s(new MouseEvent(\"%s\", {%s}))"), *BaseDispatch, *MouseEventName, *BaseMouseEventData);
+			
+			if(PointerEventName.Len() > 0)
+			{
+				InWindow.ExecuteJavascript(DispatchEventJS);
+				
+				DispatchEventJS = FString::Printf(TEXT("%s(new PointerEvent(\"%s\", {%s}))"), *BaseDispatch, *PointerEventName, *BaseMouseEventData);
+			}
+		}
+	}
+	
+	InWindow.ExecuteJavascript(DispatchEventJS);
+}
+
+void FWebView3DInputHandler::DispatchInputEvent(IWebBrowserWindow& InWindow, const FString& InputType, TCHAR* InputData)
+{
+	FString CurrentJSElement = GetCurrentJSElement();
+	FString InputDataString = InputData == nullptr ? FString(TEXT("null")) : FString(1, InputData);
+	FString BaseInputEvent = FString::Printf(TEXT("data: \"%s\", inputType: \"%s\", isComposing: false"), *InputDataString, *InputType);
+	FString BeforeInputJS = FString::Printf(TEXT("%s.dispatchEvent(new InputEvent(\"beforeinput\", {%s}))"), *CurrentJSElement, *BaseInputEvent);
+	
+	InWindow.ExecuteJavascript(BeforeInputJS);
+	
+	if(InputType == "insertText")
+	{
+		FString ContentJS = FString::Printf(TEXT("%s.value += \"%s\";"), *CurrentJSElement, *InputDataString);
+		
+		InWindow.ExecuteJavascript(ContentJS);
+	}
+	else if(InputType == "deleteContentBackward")
+	{
+		FString ContentJS = FString::Printf(TEXT("%s.value = %s.value.slice(0, -1);"), *CurrentJSElement, *CurrentJSElement);
+		
+		InWindow.ExecuteJavascript(ContentJS);
+	}
+	
+	FString InputJS = FString::Printf(TEXT("%s.dispatchEvent(new InputEvent(\"input\", {%s}))"), *CurrentJSElement, *BaseInputEvent);
+	
+	InWindow.ExecuteJavascript(InputJS);
+}
+
+void FWebView3DInputHandler::SendCharEventToJS(IWebBrowserWindow& InWindow,
+											   const FCharacterEvent& CharEvent) 
+{
+	TCHAR InputChar = CharEvent.GetCharacter();
+	
+	if(CheckIfCharacterIsInputable(InputChar))
+	{
+		DispatchInputEvent(InWindow, "insertText", &InputChar);
+	}
+	else
+	{
+		UE_LOG(LogWebView3DInputHandler, Verbose, TEXT("Uninputable character received, ignoring"));
+	}
+}
+
+void FWebView3DInputHandler::UpdateModifierKeys(const FKey& Key, bool KeyDown)
+{
+	if (Key == EKeys::LeftControl || Key == EKeys::RightControl)
+	{
+		bCtrlKeyDown = KeyDown;
+	}
+	else if (Key == EKeys::LeftAlt || Key == EKeys::RightAlt)
+	{
+		bAltKeyDown = KeyDown;
+	}
+	else if (Key == EKeys::LeftCommand || Key == EKeys::RightCommand)
+	{
+		bCmdKeyDown = KeyDown;
+	}
+	else if (Key == EKeys::LeftShift || Key == EKeys::RightShift)
+	{
+		bShiftKeyDown = KeyDown;
+	}
+}
+
+void FWebView3DInputHandler::SendKeyEventToJS(IWebBrowserWindow& InWindow, const FString& EventName, const FKeyEvent& KeyEvent) 
+{
+	// check for supported key event type
+	check(EventName == "keyup" || EventName == "keydown");
+	
+	FString CurrentJSElement = GetCurrentJSElement();
+	FKey Key = KeyEvent.GetKey();
+	uint32 KeyEventCharacter = KeyEvent.GetCharacter();
+	uint32 KeyEventCode = KeyEvent.GetKeyCode();
+	
+	FString DispatchEventJS;
+	
+	if(KeyEventCharacter == 0)
+	{
+#if PLATFORM_MAC
+		if(Key.IsModifierKey())
+		{
+			KeyEventCode = FMacPlatformInput::GetMacNativeModifierKeyCode(KeyEventCode);
+		}
+#endif
+		
+		DispatchEventJS = FString::Printf(TEXT("%s.dispatchEvent(new KeyboardEvent(\"%s\", {code: %d, %s}));"), 
+										  *CurrentJSElement, *EventName, KeyEventCode, *GetModifierKeyStatus());
+	}
+	else
+	{
+		DispatchEventJS = FString::Printf(TEXT("%s.dispatchEvent(new KeyboardEvent(\"%s\", {code: %d, key: \"%c\", %s}));"), 
+										  *CurrentJSElement, *EventName, KeyEventCode, KeyEventCharacter, *GetModifierKeyStatus());
+	}
+
+	InWindow.ExecuteJavascript(DispatchEventJS);
+	
+	bool KeyDown = EventName == "keydown";
+	
+	if (KeyDown && (Key == EKeys::BackSpace || Key == EKeys::Delete))  
+	{
+		DispatchInputEvent(InWindow, "deleteContentBackward", nullptr);
+	}
+	
+	UpdateModifierKeys(Key, KeyDown);
+}
diff --git a/Engine/Source/Runtime/WebBrowser/Private/Apple/WebView3DInputHandler.h b/Engine/Source/Runtime/WebBrowser/Private/Apple/WebView3DInputHandler.h
new file mode 100644
--- /dev/null
+++ b/Engine/Source/Runtime/WebBrowser/Private/Apple/WebView3DInputHandler.h
@@ -0,0 +1,47 @@
+// Copyright Epic Games, Inc. All Rights Reserved.
+
+#pragma once
+
+#include "CoreMinimal.h"
+
+#include "Input/Events.h"
+
+DECLARE_LOG_CATEGORY_EXTERN(LogWebView3DInputHandler, Log, All);
+
+class IWebBrowserWindow;
+
+/**
+* Helper class for ApplePlatformWebBrowser's FWebBrowserWindow to handle input with 3D WKWebView
+* Web Browser widgets.
+ * @warning Due to WebKit limitations, all input events dispatched in 3D mode will be untrusted (i.e. Event.isTrusted will be false). 
+*/
+class FWebView3DInputHandler 
+{
+public:
+	
+	FWebView3DInputHandler();
+	
+	void SendMouseEventToJS(IWebBrowserWindow& InWindow,
+							const FString& EventName, const FPointerEvent& MouseEvent);
+	void SendCharEventToJS(IWebBrowserWindow& InWindow, const FCharacterEvent& CharEvent);
+	void SendKeyEventToJS(IWebBrowserWindow& InWindow, const FString& EventName, const FKeyEvent& KeyEvent);
+	
+private:
+	
+	FString BoolToString(bool InBool);
+	FString GetCurrentJSElement();
+	
+	void DispatchInputEvent(IWebBrowserWindow& InWindow, const FString& InputType, TCHAR* InputData);
+	
+	FString GetModifierKeyStatus();
+	bool CheckIfCharacterIsInputable(TCHAR CharToInput);
+	
+	void UpdateModifierKeys(const FKey& Key, bool KeyDown);
+	
+	bool bCtrlKeyDown;
+	bool bAltKeyDown;
+	bool bCmdKeyDown;
+	bool bShiftKeyDown;
+	
+	FVector2f MousePos;
+};
