Fixing edge cases for 3D UI Input in WebView3DInputHandler (specifically escaping characters and character capitalization)

diff --git a/Engine/Source/Runtime/WebBrowser/Private/Apple/ApplePlatformWebBrowser.cpp b/Engine/Source/Runtime/WebBrowser/Private/Apple/ApplePlatformWebBrowser.cpp
--- a/Engine/Source/Runtime/WebBrowser/Private/Apple/ApplePlatformWebBrowser.cpp
+++ b/Engine/Source/Runtime/WebBrowser/Private/Apple/ApplePlatformWebBrowser.cpp
@@ -180,8 +180,9 @@ class SAppleWebBrowserWidget : public SViewport
 		}
 #endif
 
-		if(NewFrame.size.width != WindowFrame.size.width || NewFrame.size.height != WindowFrame.size.height) 
+		if(!Is3DBrowser && (NewFrame.size.width != WindowFrame.size.width || NewFrame.size.height != WindowFrame.size.height)) 
 		{
+			// WindowFrame is only used when we are a 2D web browser
 			WindowFrame = NewFrame;
 		}
 		
diff --git a/Engine/Source/Runtime/WebBrowser/Private/Apple/WebView3DInputHandler.cpp b/Engine/Source/Runtime/WebBrowser/Private/Apple/WebView3DInputHandler.cpp
--- a/Engine/Source/Runtime/WebBrowser/Private/Apple/WebView3DInputHandler.cpp
+++ b/Engine/Source/Runtime/WebBrowser/Private/Apple/WebView3DInputHandler.cpp
@@ -15,6 +15,7 @@ FWebView3DInputHandler::FWebView3DInputHandler()
 	, bAltKeyDown(false)
 	, bCmdKeyDown(false) 
 	, bShiftKeyDown(false)
+	, bCapsLockOn(false)
 {
 }
 
@@ -148,13 +149,30 @@ void FWebView3DInputHandler::SendMouseEventToJS(IWebBrowserWindow& InWindow, con
 void FWebView3DInputHandler::DispatchInputEvent(IWebBrowserWindow& InWindow, const FString& InputType, TCHAR* InputData)
 {
 	FString CurrentJSElement = GetCurrentJSElement();
-	FString InputDataString = InputData == nullptr ? FString(TEXT("null")) : FString(1, InputData);
-	FString BaseInputEvent = FString::Printf(TEXT("data: \"%s\", inputType: \"%s\", isComposing: false"), *InputDataString, *InputType);
+	FString InputDataString;
+	FString BaseInputEvent;
+	
+	if(InputData != nullptr)
+	{
+		InputDataString = FString(1, InputData).ReplaceCharWithEscapedChar();
+		
+		if(bCapsLockOn || bShiftKeyDown)
+		{
+			InputDataString = InputDataString.ToUpper();
+		}
+		
+		BaseInputEvent = FString::Printf(TEXT("data: \"%s\", inputType: \"%s\", isComposing: false"), *InputDataString, *InputType);
+	}
+	else
+	{
+		BaseInputEvent = FString::Printf(TEXT("data: null, inputType: \"%s\", isComposing: false"), *InputType);
+	}
+	
 	FString BeforeInputJS = FString::Printf(TEXT("%s.dispatchEvent(new InputEvent(\"beforeinput\", {%s}))"), *CurrentJSElement, *BaseInputEvent);
 	
 	InWindow.ExecuteJavascript(BeforeInputJS);
 	
-	if(InputType == "insertText")
+	if(InputType == "insertText" && InputData != nullptr)
 	{
 		FString ContentJS = FString::Printf(TEXT("%s.value += \"%s\";"), *CurrentJSElement, *InputDataString);
 		
@@ -205,6 +223,10 @@ void FWebView3DInputHandler::UpdateModifierKeys(const FKey& Key, bool KeyDown)
 	{
 		bShiftKeyDown = KeyDown;
 	}
+	else if (Key == EKeys::CapsLock)
+	{
+		bCapsLockOn = KeyDown;
+	}
 }
 
 void FWebView3DInputHandler::SendKeyEventToJS(IWebBrowserWindow& InWindow, const FString& EventName, const FKeyEvent& KeyEvent) 
@@ -213,13 +235,16 @@ void FWebView3DInputHandler::SendKeyEventToJS(IWebBrowserWindow& InWindow, const
 	check(EventName == "keyup" || EventName == "keydown");
 	
 	FString CurrentJSElement = GetCurrentJSElement();
+	FString DispatchEventJS;
 	FKey Key = KeyEvent.GetKey();
-	uint32 KeyEventCharacter = KeyEvent.GetCharacter();
+	bool KeyDown = EventName == "keydown";
+	uint32 CharacterKeyU32 = KeyEvent.GetCharacter();
 	uint32 KeyEventCode = KeyEvent.GetKeyCode();
 	
-	FString DispatchEventJS;
+	// modifier keys should be updated before we construct the event to dispatch
+	UpdateModifierKeys(Key, KeyDown);
 	
-	if(KeyEventCharacter == 0)
+	if(CharacterKeyU32 == 0)
 	{
 #if PLATFORM_MAC
 		if(Key.IsModifierKey())
@@ -233,18 +258,16 @@ void FWebView3DInputHandler::SendKeyEventToJS(IWebBrowserWindow& InWindow, const
 	}
 	else
 	{
-		DispatchEventJS = FString::Printf(TEXT("%s.dispatchEvent(new KeyboardEvent(\"%s\", {code: %d, key: \"%c\", %s}));"), 
-										  *CurrentJSElement, *EventName, KeyEventCode, KeyEventCharacter, *GetModifierKeyStatus());
+		FString KeyEventCharacter = FString(reinterpret_cast<TCHAR*>(&CharacterKeyU32), 1).ReplaceCharWithEscapedChar();
+		
+		DispatchEventJS = FString::Printf(TEXT("%s.dispatchEvent(new KeyboardEvent(\"%s\", {code: %d, key: \"%s\", %s}));"), 
+										  *CurrentJSElement, *EventName, KeyEventCode, *KeyEventCharacter, *GetModifierKeyStatus());
 	}
 
 	InWindow.ExecuteJavascript(DispatchEventJS);
 	
-	bool KeyDown = EventName == "keydown";
-	
 	if (KeyDown && (Key == EKeys::BackSpace || Key == EKeys::Delete))  
 	{
 		DispatchInputEvent(InWindow, "deleteContentBackward", nullptr);
 	}
-	
-	UpdateModifierKeys(Key, KeyDown);
 }
diff --git a/Engine/Source/Runtime/WebBrowser/Private/Apple/WebView3DInputHandler.h b/Engine/Source/Runtime/WebBrowser/Private/Apple/WebView3DInputHandler.h
--- a/Engine/Source/Runtime/WebBrowser/Private/Apple/WebView3DInputHandler.h
+++ b/Engine/Source/Runtime/WebBrowser/Private/Apple/WebView3DInputHandler.h
@@ -42,6 +42,7 @@ private:
 	bool bAltKeyDown;
 	bool bCmdKeyDown;
 	bool bShiftKeyDown;
+	bool bCapsLockOn;
 	
 	FVector2f MousePos;
 };
