[WebBrowser]

Provides a workaround for this known CEF bug -
[REDACTED_PROFILE_URL]/cef/issues/3826

This is turned on by default in this CL!
To turn it off by default, this line can be changed in SWebBrowser.h-
_UseSmoothResizing(false)
And if it is made to be off by default, then individual developers will need to turn it on themselves as needed, i.e -
SAssignNew(WebBrowserWidget, SWebBrowser, WebBrowserWindow).UseSmoothResizing(true)

A new CapturedCefBuffer class is provided, which stores the last 'good' buffer received from CEF, along with the ability to paint it via Slate.
FCEFWebBrowserWindow::OnPaint() updates this CapturedCefBuffer whenever it receives 'good' CEF buffers.
SWebBrowser determines whether FCEFWebBrowserWindow is able to properly paint itself, or if it should paint this CapturedCefBuffer instead.
There were additional issues with making FCEFWebBrowserWindow know that it needs to resize it's CEF image, also making WebBrowserSingleton to tick when Slate is suspended during active window resizes, which were accomplished using SWebBrowser::Tick(), which does not get suspended.

diff --git a/Engine/Source/Runtime/WebBrowser/Private/CEF/CEFWebBrowserWindow.cpp b/Engine/Source/Runtime/WebBrowser/Private/CEF/CEFWebBrowserWindow.cpp
--- a/Engine/Source/Runtime/WebBrowser/Private/CEF/CEFWebBrowserWindow.cpp
+++ b/Engine/Source/Runtime/WebBrowser/Private/CEF/CEFWebBrowserWindow.cpp
@@ -243,6 +243,8 @@ enum KeyboardCode {
 #define USE_BUFFERED_VIDEO 0
 #endif
 
+#define UE_CEF_HAS_RESIZE_BUG (CEF_VERSION_MAJOR >= 128)
+
 namespace {
 	// Private helper class to post a callback to GetSource.
 	class FWebBrowserClosureVisitor
@@ -1980,6 +1982,21 @@ void FCEFWebBrowserWindow::OnPaint(CefRenderHandler::PaintElementType Type, cons
 		return;
 	}
 
+
+#if UE_CEF_HAS_RESIZE_BUG
+	const FIntPoint PaintedBufferSize(Width / ViewportDPIScaleFactor, Height / ViewportDPIScaleFactor);
+	bHasCorrectNativeCefBuffer = (PaintedBufferSize == ViewportSize);
+	if (IsUsingCapturedCefBuffer()) 
+	{
+		// We attempt to capture a good CEF buffer, especially when it is correct. The CEF buffer size does not actually need to match the viewport
+		// size to end up being captured, but also may not be captured if it fails another test (which we skip if we know we have the correct CEF
+		// buffer.) See code.
+
+		CapturedCefBuffer->SetBufferAsB8G8R8A8(Buffer, Width, Height, ViewportDPIScaleFactor, bHasCorrectNativeCefBuffer);
+	}
+#endif
+
+	
 	bool bNeedsRedraw = false;
 	if (bUsingAcceleratedPaint)
 	{
@@ -2022,8 +2039,8 @@ void FCEFWebBrowserWindow::OnPaint(CefRenderHandler::PaintElementType Type, cons
 			// This is another side-effect of [REDACTED_PROFILE_URL]/cef/issues/3826:
 			// When resizing the view, the last frame we receive occasionally still has a previous viewport size,
 			// which results in contents looking stretched or with black borders.
-#if CEF_VERSION_MAJOR >= 128
-			if (Type == PET_VIEW && ViewportSize != FIntPoint(Width / ViewportDPIScaleFactor, Height / ViewportDPIScaleFactor))
+#if UE_CEF_HAS_RESIZE_BUG
+			if (!IsUsingCapturedCefBuffer() && Type == PET_VIEW && !bHasCorrectNativeCefBuffer)
 			{
 				// Note that calling Invalidate() doesn't work here, and there's no other CEF API to trigger a repaint.
 				InternalCefBrowser->GetHost()->WasHidden(true);
@@ -2711,7 +2728,7 @@ void FCEFWebBrowserWindow::CheckTickActivity()
 		// Recent CEF versions fail to trigger a repaint after WasResized() is called
 		// so we must force one on the view (and the popup if visible):
 		// (upstream issue: [REDACTED_PROFILE_URL]/cef/issues/3826)
-#if CEF_VERSION_MAJOR >= 128
+#if UE_CEF_HAS_RESIZE_BUG
 		InternalCefBrowser->GetHost()->Invalidate(PET_VIEW);
 		if (bPopupHasFocus)
 		{
@@ -2900,6 +2917,53 @@ CefRefPtr<CefDictionaryValue> FCEFWebBrowserWindow::GetProcessInfo()
 	return Retval;
 }
 
+bool FCEFWebBrowserWindow::HasCorrectNativeCefBuffer() const
+{
+#if UE_CEF_HAS_RESIZE_BUG
+	return bHasCorrectNativeCefBuffer;
+#else
+	return true;
+#endif
+}
+
+void FCEFWebBrowserWindow::SetUseCapturedCefBuffer(const bool bValue)
+{
+#if UE_CEF_HAS_RESIZE_BUG
+	if (bValue)
+	{
+		if (!CapturedCefBuffer.IsValid())
+		{
+			CapturedCefBuffer = MakeShared<FCapturedCefBuffer>();
+		}
+	}
+	else
+	{
+		CapturedCefBuffer.Reset();
+	}
+#else
+	CapturedCefBuffer.Reset();
+#endif
+}
+
+bool FCEFWebBrowserWindow::IsUsingCapturedCefBuffer() const
+{
+#if UE_CEF_HAS_RESIZE_BUG
+	return CapturedCefBuffer.IsValid();
+#else
+	return false;
+#endif
+}
+
+TSharedPtr<FCapturedCefBuffer> FCEFWebBrowserWindow::GetCapturedCefBuffer()
+{
+	return CapturedCefBuffer; // ..may not be set
+}
+
+void FCEFWebBrowserWindow::SetNeedsResize(const bool bInValue)
+{
+	bNeedsResize = bInValue;
+}
+
 bool FCEFWebBrowserWindow::OnProcessMessageReceived(CefRefPtr<CefBrowser> Browser, CefRefPtr<CefFrame> frame, CefProcessId SourceProcess, CefRefPtr<CefProcessMessage> Message)
 {
 	if (IsClosing())
diff --git a/Engine/Source/Runtime/WebBrowser/Private/CEF/CEFWebBrowserWindow.h b/Engine/Source/Runtime/WebBrowser/Private/CEF/CEFWebBrowserWindow.h
--- a/Engine/Source/Runtime/WebBrowser/Private/CEF/CEFWebBrowserWindow.h
+++ b/Engine/Source/Runtime/WebBrowser/Private/CEF/CEFWebBrowserWindow.h
@@ -14,9 +14,9 @@
 #include "IWebBrowserWindow.h"
 #include "CEFBrowserHandler.h"
 
-
 #include "CEFLibCefIncludes.h"
 
+#include "CapturedCefBuffer.h"
 
 #endif
 
@@ -576,6 +576,32 @@ public:
 	*/
 	bool URLRequestAllowsCredentials(const FString& URL) const { return WebBrowserHandler->URLRequestAllowsCredentials(URL); }
 
+	/**
+	 * Tells us whether this has the correct CEF image buffer.
+	 */
+	WEBBROWSER_API bool HasCorrectNativeCefBuffer() const;
+	
+	/**
+	 * Whether we should use a workaround for this bug -
+	 * [REDACTED_PROFILE_URL]/cef/issues/3826
+	 */
+	WEBBROWSER_API void SetUseCapturedCefBuffer(const bool bValue);
+	
+	/**
+	 * Whether we are using the captured CEF buffer.
+	 */
+	WEBBROWSER_API bool IsUsingCapturedCefBuffer() const;
+	
+	/**
+	 * Gets an undistorted buffer captured from CEF. Format is PF_B8G8R8A8.
+	 */
+	WEBBROWSER_API TSharedPtr<FCapturedCefBuffer> GetCapturedCefBuffer();
+
+	/**
+	 * Set if a resize is needed.
+	 */
+	WEBBROWSER_API void SetNeedsResize(const bool bInValue);
+	
 private:
 
 	/** @return the currently valid renderer, if available */
@@ -806,6 +832,9 @@ private:
 #if PLATFORM_WINDOWS || PLATFORM_MAC
 	bool bInDirectHwndMode;
 #endif
+
+	bool bHasCorrectNativeCefBuffer = true;
+	TSharedPtr<FCapturedCefBuffer> CapturedCefBuffer;
 };
 
 typedef FCEFWebBrowserWindow FWebBrowserWindow;
diff --git a/Engine/Source/Runtime/WebBrowser/Private/CapturedCefBuffer.cpp b/Engine/Source/Runtime/WebBrowser/Private/CapturedCefBuffer.cpp
new file mode 100644
--- /dev/null
+++ b/Engine/Source/Runtime/WebBrowser/Private/CapturedCefBuffer.cpp
@@ -0,0 +1,208 @@
+// Copyright Epic Games, Inc. All Rights Reserved.
+
+
+#include "CapturedCefBuffer.h"
+
+#include "Runtime/Engine/Public/TextureResource.h"
+#include "Framework/Application/SlateApplication.h"
+#include "Styling/StyleColors.h"
+
+
+bool FCapturedCefBuffer::SetBufferAsB8G8R8A8(const void* InBufferB8G8R8A8, const int32 InBufferWidth, const int32 InBufferHeight, const float InViewportDPIScaleFactor, const bool bDoSkipBadBufferTest)
+{
+	if (!InBufferB8G8R8A8 || InBufferWidth <= 0 || InBufferHeight <= 0)
+	{
+		ClearBuffer();
+		
+		return false;
+	}
+
+
+	// This works depending on whether the web browser content has a transparent background.
+	// If it does, then the web browser will show its content, centered and stable.
+	// If it doesn't, then the web browser will occasionally adapt so that its content fills the window, depending on this -
+	//
+	// CEF can sometimes send us a buffer that's apparently a good size. But, the actual image pixels of the web browser content are actually
+	// somewhere in the middle of that image, smaller than the entire image actually is, stretched or compressed, and all surrounded with a margin of
+	// transparent pixels. To test if this is the case at all, we really only need to test the alpha channel of a single pixel in one of the corners,
+	// and the very first pixel will be one such pixel. If this is the case, it's a dead giveaway that the buffer is not good, and so we don't want to
+	// capture this CEF buffer, because the web browser image will appear to be the wrong size and stretched or compressed. We'll wait until a buffer
+	// comes along where the web browser image is actually completely filled and not stretched or compressed, which is the case when there's no
+	// transparent margin.
+	//
+	// We also don't capture the buffer if the dimensions have not changed, to avoid captures on every frame.
+	//
+	// But we allow the caller to skip this test and force capture, especially if the caller knows they have a correct buffer.
+
+	if (InBufferWidth == BufferDimensions.X && InBufferHeight == BufferDimensions.Y && InViewportDPIScaleFactor == ViewportDPIScaleFactor)
+	{
+		return false;
+	}
+	
+	if (const bool bDoesBufferHaveTransparentMargin = (static_cast<const uint8*>(InBufferB8G8R8A8)[3] == 0x0); 
+		bDoSkipBadBufferTest || (!bDoSkipBadBufferTest && !bDoesBufferHaveTransparentMargin))
+	{
+		BufferDimensions = FIntPoint(InBufferWidth, InBufferHeight);
+		
+		const TArray<uint8>::SizeType NumBytes = CalculateBufferNumBytes();
+		if (BufferData.Num() != NumBytes)
+		{
+			BufferData.SetNum(NumBytes);
+		}
+		FMemory::Memcpy(BufferData.GetData(), InBufferB8G8R8A8, NumBytes);
+
+
+		ViewportDPIScaleFactor = InViewportDPIScaleFactor;
+
+
+		bDoesNeedNewPaintObjects = true;
+
+		
+		return true;
+	}
+
+	
+	return false;
+}
+
+
+bool FCapturedCefBuffer::ClearBuffer()
+{
+	if (!BufferData.IsEmpty())
+	{
+		BufferData.Empty();
+	}
+	BufferDimensions = FIntPoint::ZeroValue;
+	ViewportDPIScaleFactor = 1.0f;
+
+	
+	return true;
+}
+
+
+bool FCapturedCefBuffer::MaybeUpdatePaintObjects() 
+{
+	if (!FSlateApplication::IsInitialized()) // ..safety check
+	{
+		return false;
+	}
+	
+	if (!IsBufferValid())
+	{
+		ClearPaintObjects();
+		
+		return false;
+	}
+
+	
+	// Copy buffer into texture, and into Slate brush.
+
+	if (bDoesNeedNewPaintObjects)
+	{
+		if (PaintTexture.Get() && (PaintTexture->GetSizeX() != BufferDimensions.X || PaintTexture->GetSizeY() != BufferDimensions.Y))
+		{
+			PaintTexture.Reset(); // ..uses GC
+		}
+
+		if (!PaintTexture.IsValid())
+		{
+			PaintTexture = TStrongObjectPtr(UTexture2D::CreateTransient(BufferDimensions.X, BufferDimensions.Y, PF_B8G8R8A8));
+			PaintTexture->SRGB = true;
+			PaintTexture->LODGroup = TEXTUREGROUP_UI;
+		}
+		
+		void* TextureData = PaintTexture->GetPlatformData()->Mips[0].BulkData.Lock(LOCK_READ_WRITE);
+		FMemory::Memcpy(TextureData, BufferData.GetData(), CalculateBufferNumBytes());
+		PaintTexture->GetPlatformData()->Mips[0].BulkData.Unlock();
+		PaintTexture->UpdateResource();
+	
+	
+		PaintSlateBrush.SetResourceObject(PaintTexture.Get());
+		PaintSlateBrush.ImageSize = FVector2D(BufferDimensions.X, BufferDimensions.Y);
+		PaintSlateBrush.DrawAs = ESlateBrushDrawType::Type::Image;
+
+		
+		bDoesNeedNewPaintObjects = false;
+		bHasPaintObjects = true;
+	}
+
+	
+	return true;
+}
+
+
+bool FCapturedCefBuffer::ClearPaintObjects()
+{
+	if (!bHasPaintObjects)
+	{
+		return false;
+	}
+	
+	if (!FSlateApplication::IsInitialized()) // ..safety check
+	{
+		return false;
+	}
+
+	
+	bHasPaintObjects = false;
+	
+	if (FSlateApplication::Get().GetRenderer()/*becomes null on shutdown*/ && PaintSlateBrush.GetResourceObject())
+	{
+		FSlateApplication::Get().GetRenderer()->ReleaseDynamicResource(PaintSlateBrush);
+	}
+
+	if (PaintTexture.Get())
+	{
+		PaintTexture.Reset(); // ..uses GC
+	}
+
+
+	return true;
+}
+
+
+int32 FCapturedCefBuffer::PaintCentered(const FGeometry& AllottedGeometry, FSlateWindowElementList& OutDrawElements, int32 LayerId, const FWidgetStyle& InWidgetStyle) const
+{
+	if (!bHasPaintObjects || ViewportDPIScaleFactor <= 0.0)
+	{
+		return LayerId;
+	}
+
+
+	const float Zoom = 1.0f / ViewportDPIScaleFactor; // ..was checked for divide by zero
+
+	const FVector2D& SlateBrushSize = PaintSlateBrush.ImageSize;
+	const FVector2D WidgetSize = AllottedGeometry.GetLocalSize();
+	const FVector2D CenteringOffset = (WidgetSize - SlateBrushSize) * 0.5f;
+	const FSlateLayoutTransform CenteringTransform(Zoom, Zoom * CenteringOffset);
+
+	FSlateDrawElement::MakeBox(
+		OutDrawElements,
+		LayerId++,
+		AllottedGeometry.ToPaintGeometry(SlateBrushSize, CenteringTransform),
+		&PaintSlateBrush,
+		ESlateDrawEffect::None,
+		InWidgetStyle.GetColorAndOpacityTint()
+	);
+
+	
+	return LayerId;
+}
+
+
+bool FCapturedCefBuffer::HasPaintObjects() const
+{
+	return bHasPaintObjects;
+}
+
+
+bool FCapturedCefBuffer::IsBufferValid() const
+{
+	return (!BufferData.IsEmpty() && BufferDimensions.X > 0 && BufferDimensions.Y > 0);	
+}
+
+
+uint32 FCapturedCefBuffer::CalculateBufferNumBytes() const
+{
+	return (BufferDimensions.X * BufferDimensions.Y * 4/*PF_B8G8R8A8*/);
+}
diff --git a/Engine/Source/Runtime/WebBrowser/Private/CapturedCefBuffer.h b/Engine/Source/Runtime/WebBrowser/Private/CapturedCefBuffer.h
new file mode 100644
--- /dev/null
+++ b/Engine/Source/Runtime/WebBrowser/Private/CapturedCefBuffer.h
@@ -0,0 +1,50 @@
+// Copyright Epic Games, Inc. All Rights Reserved.
+
+
+#pragma once
+
+
+#include "Engine/Texture2D.h"
+#include "UObject/StrongObjectPtrTemplates.h"
+#include "Layout/Geometry.h"
+#include "Rendering/DrawElements.h"
+#include "Styling/SlateBrush.h"
+
+#include "CapturedCefBuffer.generated.h"
+
+
+USTRUCT()
+struct FCapturedCefBuffer
+{
+	GENERATED_BODY()
+
+
+public:
+
+
+	bool SetBufferAsB8G8R8A8(const void* InBufferB8G8R8A8, const int32 InBufferWidth, const int32 InBufferHeight, const float InViewportDPIScaleFactor, const bool bDoSkipBadBufferTest);
+	bool ClearBuffer();
+	
+	bool MaybeUpdatePaintObjects();
+	bool ClearPaintObjects();
+	int32 PaintCentered(const FGeometry& AllottedGeometry, FSlateWindowElementList& OutDrawElements, int32 LayerId, const FWidgetStyle& InWidgetStyle) const;
+
+	bool HasPaintObjects() const;
+
+private:
+
+
+	bool IsBufferValid() const; 
+	uint32 CalculateBufferNumBytes() const;
+	
+
+	TArray<uint8> BufferData;
+	FIntPoint BufferDimensions = FIntPoint::ZeroValue;
+
+	float ViewportDPIScaleFactor = 1.0f;
+	
+	bool bDoesNeedNewPaintObjects = false;
+	bool bHasPaintObjects = false;
+	TStrongObjectPtr<UTexture2D> PaintTexture;
+	FSlateBrush PaintSlateBrush;
+};
diff --git a/Engine/Source/Runtime/WebBrowser/Private/SWebBrowser.cpp b/Engine/Source/Runtime/WebBrowser/Private/SWebBrowser.cpp
--- a/Engine/Source/Runtime/WebBrowser/Private/SWebBrowser.cpp
+++ b/Engine/Source/Runtime/WebBrowser/Private/SWebBrowser.cpp
@@ -5,6 +5,14 @@
 #include "Widgets/Input/SEditableTextBox.h"
 #include "Widgets/Input/SButton.h"
 #include "Widgets/Images/SThrobber.h"
+#if WITH_CEF3
+#include "Framework/Application/SlateApplication.h"
+#include "Framework/Docking/TabManager.h"
+#include "Async/Async.h"
+#include "CEF/CEFWebBrowserWindow.h"
+#endif
+#include "WebBrowserModule.h"
+#include "WebBrowserSingleton.h"
 
 #define LOCTEXT_NAMESPACE "WebBrowser"
 
@@ -16,7 +24,151 @@ SWebBrowser::~SWebBrowser()
 {
 }
 
-void SWebBrowser::Construct(const FArguments& InArgs, const TSharedPtr<IWebBrowserWindow>& BrowserWindow)
+#if WITH_CEF3
+
+static TSharedRef<FCEFWebBrowserWindow> GetCefWebBrowserWindowChecked(const TSharedPtr<IWebBrowserWindow>& WebBrowserWindow)
+{
+	check(WebBrowserWindow.IsValid());
+	
+	const TSharedPtr<FCEFWebBrowserWindow> CefWebBrowserWindow = StaticCastSharedPtr<FCEFWebBrowserWindow>(WebBrowserWindow);
+	check(CefWebBrowserWindow.IsValid());
+	
+	return CefWebBrowserWindow.ToSharedRef();
+}
+
+#endif
+
+static FWebBrowserSingleton& GetWebBrowserSingletonChecked()
+{
+	IWebBrowserModule& WebBrowserModule = FModuleManager::LoadModuleChecked<IWebBrowserModule>("WebBrowser");
+	FWebBrowserSingleton* WebBrowserSingleton = static_cast<FWebBrowserSingleton*>(WebBrowserModule.GetSingleton());
+	check(WebBrowserSingleton);
+	return *WebBrowserSingleton;
+}
+
+/*virtual*/ void SWebBrowser::Tick(const FGeometry& AllottedGeometry, const double InCurrentTime, const float InDeltaTime)
+{
+	SCompoundWidget::Tick(AllottedGeometry, InCurrentTime, InDeltaTime);
+
+	
+#if WITH_EDITOR
+#if WITH_CEF3
+
+	// If we're using captured CEF buffer / smooth resizing.
+
+	if (WebBrowserWindow.IsValid()) 
+	{
+		if (const TSharedRef<FCEFWebBrowserWindow> CefWebBrowserWindow = GetCefWebBrowserWindowChecked(WebBrowserWindow);
+			CefWebBrowserWindow->IsUsingCapturedCefBuffer())
+		{
+			// NOTE - We won't end up here if the CEF version doesn't require it.
+			
+			const bool bIsFloatingWindow = (FSlateApplication::Get().FindWidgetWindow(AsShared()) != FGlobalTabmanager::Get()->GetRootWindow());
+			
+
+			// * This is necessary during active widget resizes, whether or not the widget is in a floating window. *
+			// If web browser doesn't have the correct CEF buffer, tell it to resize.
+			// NOTE - Call this before the 'forced tick' code below.
+			
+			if (!CefWebBrowserWindow->HasCorrectNativeCefBuffer())
+			{
+				if (bIsFloatingWindow)
+				{
+					// Seems necessary to make web browser content appear, if it's already in a floating window when the Editor starts up.
+					// But this is not good to call if docked in editor!
+					CefWebBrowserWindow->SetViewportSize(BrowserView->GetCachedGeometry().Size.IntPoint(), BrowserView->GetCachedGeometry().AbsolutePosition.IntPoint());
+				}
+
+				// Always seems necessary (even if viewport size was just changed above.)
+				CefWebBrowserWindow->SetNeedsResize(true);
+			}
+
+	
+#if PLATFORM_WINDOWS // TODO - Test MacOS and Linux.
+
+			// * This is necessary during active widget resizes, only if the widget is in a floating window, and only when on Windows. *
+			// When we are (1) running on Windows (2) not docked in the Editor (floating), and (3) our window is actively being resized - then the FTS
+			// Ticker interface on the Web Browser Singleton will be suspended. But our Tick() function here will not be suspended. So under these and
+			// certain other conditions (see below), we want to force Web Browser Singleton to tick, and we can do that from here. But, we also can't tick
+			// it too fast or too often, or we'll overwhelm the CEF message pump, and it will start ignoring what we want it to do. We also can't tick it
+			// too slow, or it can ignore the tick.
+			// TODO - In this floating window case only, the web browser can still end up stuck with the wrong texture, when mouse stops moving during
+			//		  an active resize, until mouse is released. Why?
+			
+			if (FWebBrowserSingleton& WebBrowserSingleton = GetWebBrowserSingletonChecked();
+				bIsFloatingWindow && FPlatformTime::Seconds() - WebBrowserSingleton.GetPreviousTickTimeSeconds() >= WebBrowserSingleton.GetExternalTickCallWaitSeconds())
+			{
+				WebBrowserSingleton.Tick(InDeltaTime);
+			}
+	
+#endif
+		}
+	}
+	
+#endif
+#endif
+}
+
+/*virtual*/ int32 SWebBrowser::OnPaint(const FPaintArgs& Args, const FGeometry& AllottedGeometry, const FSlateRect& MyCullingRect, FSlateWindowElementList& OutDrawElements, int32 LayerId, const FWidgetStyle& InWidgetStyle, bool bParentEnabled) const
+{
+#if WITH_EDITOR
+#if WITH_CEF3
+	
+	// If we're using captured CEF buffer / smooth resizing..
+
+	if (WebBrowserWindow.IsValid())
+	{
+		if (const TSharedPtr<FCEFWebBrowserWindow> CefWebBrowserWindow = GetCefWebBrowserWindowChecked(WebBrowserWindow);
+			CefWebBrowserWindow->IsUsingCapturedCefBuffer())
+		{
+			// NOTE - We won't end up here if the CEF version doesn't require it.
+			
+			if (const TSharedPtr<FCapturedCefBuffer>& CapturedCefBuffer = CefWebBrowserWindow->GetCapturedCefBuffer();
+				CapturedCefBuffer.IsValid())
+			{
+				if (CefWebBrowserWindow->HasCorrectNativeCefBuffer())
+				{
+					// The CEF web browser can display the correct texture itself. So clean-up the captured CEF buffer's paint objects, if it has any.
+					// But do this on the next tick to avoid a white flash.
+
+					if (CapturedCefBuffer->HasPaintObjects())
+					{
+						AsyncTask(ENamedThreads::GameThread, [CapturedCefBuffer]() -> void
+							{
+								if (CapturedCefBuffer.IsValid())
+								{
+									CapturedCefBuffer->ClearPaintObjects();
+								}
+							});
+					}
+				}
+				else if (CapturedCefBuffer->MaybeUpdatePaintObjects())
+				{
+					// The CEF web browser can't display the correct texture itself. So paint the captured CEF buffer instead.
+					// NOTE - We still need to call the usual OnPaint() to make the web browser update, but we ignore what it actually paints. Instead, we
+					// only properly paint the captured CEF buffer.
+			
+					FSlateWindowElementList IgnoredDrawElements(FSlateApplication::Get().FindWidgetWindow(AsShared()));
+					SCompoundWidget::OnPaint(Args, AllottedGeometry, MyCullingRect, IgnoredDrawElements, LayerId, InWidgetStyle, bParentEnabled);
+
+					const FWidgetStyle CompoundedWidgetStyle = FWidgetStyle(InWidgetStyle)
+						.BlendColorAndOpacityTint(GetColorAndOpacity())
+						.SetForegroundColor(ShouldBeEnabled(bParentEnabled) ? GetForegroundColor() : GetDisabledForegroundColor() );
+
+					return CapturedCefBuffer->PaintCentered(AllottedGeometry, OutDrawElements, LayerId, CompoundedWidgetStyle);
+				}
+			}
+		}
+	}
+	
+#endif
+#endif
+
+
+	return SCompoundWidget::OnPaint(Args, AllottedGeometry, MyCullingRect, OutDrawElements, LayerId, InWidgetStyle, bParentEnabled);
+}
+
+void SWebBrowser::Construct(const FArguments& InArgs, const TSharedPtr<IWebBrowserWindow>& InWebBrowserWindow)
 {
 	OnLoadCompleted = InArgs._OnLoadCompleted;
 	OnLoadError = InArgs._OnLoadError;
@@ -102,7 +254,7 @@ void SWebBrowser::Construct(const FArguments& InArgs, const TSharedPtr<IWebBrows
 			SNew(SOverlay)
 			+ SOverlay::Slot()
 			[
-				SAssignNew(BrowserView, SWebBrowserView, BrowserWindow)
+				SAssignNew(BrowserView, SWebBrowserView, InWebBrowserWindow)
 				.ParentWindow(InArgs._ParentWindow)
 				.InitialURL(InArgs._InitialURL)
 				.ContentsToLoad(InArgs._ContentsToLoad)
@@ -141,6 +293,25 @@ void SWebBrowser::Construct(const FArguments& InArgs, const TSharedPtr<IWebBrows
 			]
 		]
 	];
+
+
+	WebBrowserWindow = (InWebBrowserWindow.IsValid() ? InWebBrowserWindow : BrowserView->GetBrowserWindow());	
+
+
+#if WITH_CEF3
+	if (WebBrowserWindow.IsValid()) 
+	{
+		const bool bUseSmoothScrolling = 
+#if WITH_EDITOR
+			(GIsPlayInEditorWorld ? false : InArgs._UseSmoothResizing)
+#else
+			false
+#endif
+			;
+		
+		GetCefWebBrowserWindowChecked(WebBrowserWindow)->SetUseCapturedCefBuffer(bUseSmoothScrolling);
+	}
+#endif
 }
 
 void SWebBrowser::LoadURL(FString NewURL)
diff --git a/Engine/Source/Runtime/WebBrowser/Private/SWebBrowserView.cpp b/Engine/Source/Runtime/WebBrowser/Private/SWebBrowserView.cpp
--- a/Engine/Source/Runtime/WebBrowser/Private/SWebBrowserView.cpp
+++ b/Engine/Source/Runtime/WebBrowser/Private/SWebBrowserView.cpp
@@ -514,6 +514,11 @@ void SWebBrowserView::ShowFloatingCloseButton(bool bShow, bool bDraggable)
 	BrowserWindow->ShowFloatingCloseButton(bShow, bDraggable);
 }
 
+TSharedPtr<IWebBrowserWindow> SWebBrowserView::GetBrowserWindow() const
+{
+	return BrowserWindow;
+}
+
 void SWebBrowserView::HandleToolTip(FString ToolTipText)
 {
 	if(ToolTipText.IsEmpty())
diff --git a/Engine/Source/Runtime/WebBrowser/Private/WebBrowserSingleton.cpp b/Engine/Source/Runtime/WebBrowser/Private/WebBrowserSingleton.cpp
--- a/Engine/Source/Runtime/WebBrowser/Private/WebBrowserSingleton.cpp
+++ b/Engine/Source/Runtime/WebBrowser/Private/WebBrowserSingleton.cpp
@@ -719,6 +719,8 @@ bool FWebBrowserSingleton::Tick(float DeltaTime)
 {
 	QUICK_SCOPE_CYCLE_COUNTER(STAT_FWebBrowserSingleton_Tick);
 
+	PreviousTickTimeSeconds = FPlatformTime::Seconds();
+	
 #if WITH_CEF3
 	if (bAllowCEF)
 	{
@@ -768,6 +770,12 @@ bool FWebBrowserSingleton::Tick(float DeltaTime)
 		float MinMessageLoopSeconds = 1.0f / MinMessageLoopHz;
 		float MaxForcedMessageLoopSeconds = 1.0f / MaxForcedMessageLoopHz;
 
+		// Set the external tick-call wait time to be in the acceptable range, but leaning heavily toward the longer time. That way, the ticker in
+		// this class is likely to tick itself first, whenever possible. If the ticker is suspended (see where this is used outside of this class),
+		// then this should be an acceptable wait time.
+		static constexpr float ExternalTickCallWaitSecondsScale = 0.95f; 
+		ExternalTickCallWaitSeconds = ExternalTickCallWaitSecondsScale * MaxForcedMessageLoopSeconds;
+		
 		static float SecondsSinceLastPump = 0;
 		static float SecondsSinceLastAppFocusCheck = MaxForcedMessageLoopSeconds;
 		static float SecondsToNextForcedPump = MaxForcedMessageLoopSeconds;
@@ -857,6 +865,21 @@ bool FWebBrowserSingleton::Tick(float DeltaTime)
 	return true;
 }
 
+#if WITH_CEF3
+double FWebBrowserSingleton::GetPreviousTickTimeSeconds() const
+{
+	return PreviousTickTimeSeconds;
+}
+#endif
+
+#if WITH_CEF3
+float FWebBrowserSingleton::GetExternalTickCallWaitSeconds() const
+{
+	static constexpr float DefaultExternalTickCallWaitSeconds = 0.064f;
+	return (ExternalTickCallWaitSeconds <= 0.0f ? DefaultExternalTickCallWaitSeconds : ExternalTickCallWaitSeconds);
+}
+#endif
+
 FString FWebBrowserSingleton::GetCurrentLocaleCode()
 {
 	FCultureRef Culture = FInternationalization::Get().GetCurrentCulture();
diff --git a/Engine/Source/Runtime/WebBrowser/Private/WebBrowserSingleton.h b/Engine/Source/Runtime/WebBrowser/Private/WebBrowserSingleton.h
--- a/Engine/Source/Runtime/WebBrowser/Private/WebBrowserSingleton.h
+++ b/Engine/Source/Runtime/WebBrowser/Private/WebBrowserSingleton.h
@@ -149,6 +149,12 @@ public:
 	virtual bool Tick(float DeltaTime) override;
 
 #if WITH_CEF3
+	/** Returns the time in seconds that an external caller should wait after previous tick, before externally calling Tick(). */
+	WEBBROWSER_API float GetExternalTickCallWaitSeconds() const;
+
+	/** Returns the time of the most recent call to Tick(). */
+	WEBBROWSER_API double GetPreviousTickTimeSeconds() const;
+
 	/** Return true if this URL will support adding an Authorization header to it */
 	bool URLRequestAllowsCredentials(const FString& URL);
 #endif
@@ -205,6 +211,11 @@ private:
 	/** Reference to UWebBrowser's translucent material*/
 	UMaterialInterface* DefaultTranslucentMaterial;
 
+	/* Time when our Tick() was last called. */
+	double PreviousTickTimeSeconds = 0.0;
+	
+	/** How long an external caller of our Tick() function should wait, after the previous tick time, before calling our Tick(). */ 
+	float ExternalTickCallWaitSeconds = 0.0;
 };
 
 PRAGMA_ENABLE_DEPRECATION_WARNINGS
diff --git a/Engine/Source/Runtime/WebBrowser/Public/SWebBrowser.h b/Engine/Source/Runtime/WebBrowser/Public/SWebBrowser.h
--- a/Engine/Source/Runtime/WebBrowser/Public/SWebBrowser.h
+++ b/Engine/Source/Runtime/WebBrowser/Public/SWebBrowser.h
@@ -38,6 +38,7 @@ public:
 		, _BackgroundColor(255,255,255,255)
 		, _BrowserFrameRate(24)
 		, _PopupMenuMethod(TOptional<EPopupMethod>())
+		, _UseSmoothResizing(true)
 		, _ViewportSize(FVector2D::ZeroVector)
 	{ }
 
@@ -77,6 +78,9 @@ public:
 		/** Override the popup menu method used for popup menus. If not set, parent widgets will be queried instead. */
 		SLATE_ARGUMENT(TOptional<EPopupMethod>, PopupMenuMethod)
 
+		/** Whether we should use a workaround for CEF bugs. */
+		SLATE_ARGUMENT(bool, UseSmoothResizing);
+		
 		/** Desired size of the web browser viewport. */
 		SLATE_ATTRIBUTE(FVector2D, ViewportSize);
 
@@ -134,6 +138,9 @@ public:
 
 	virtual bool SupportsKeyboardFocus() const override {return true;}
 
+	virtual void Tick(const FGeometry& AllottedGeometry, const double InCurrentTime, const float InDeltaTime) override;
+	virtual int32 OnPaint(const FPaintArgs& Args, const FGeometry& AllottedGeometry, const FSlateRect& MyCullingRect, FSlateWindowElementList& OutDrawElements, int32 LayerId, const FWidgetStyle& InWidgetStyle, bool bParentEnabled) const override;
+	
 	/**
 	 * Construct the widget.
 	 *
@@ -313,4 +320,6 @@ private:
 	/** The initial throbber setting */
 	bool bShowInitialThrobber;
 
+	/** The WebBrowserWindow itself. */
+	TSharedPtr<IWebBrowserWindow> WebBrowserWindow;
 };
diff --git a/Engine/Source/Runtime/WebBrowser/Public/SWebBrowserView.h b/Engine/Source/Runtime/WebBrowser/Public/SWebBrowserView.h
--- a/Engine/Source/Runtime/WebBrowser/Public/SWebBrowserView.h
+++ b/Engine/Source/Runtime/WebBrowser/Public/SWebBrowserView.h
@@ -296,6 +296,9 @@ public:
 	/** Shows floating close button.*/
 	WEBBROWSER_API void ShowFloatingCloseButton(bool bShow, bool bDraggable);
 
+	/** Gets the internal web browser window */
+	WEBBROWSER_API TSharedPtr<IWebBrowserWindow> GetBrowserWindow() const;
+	
 private:
 
 	WEBBROWSER_API void SetupParentWindowHandlers();
diff --git a/Engine/Source/Runtime/WebBrowser/WebBrowser.Build.cs b/Engine/Source/Runtime/WebBrowser/WebBrowser.Build.cs
--- a/Engine/Source/Runtime/WebBrowser/WebBrowser.Build.cs
+++ b/Engine/Source/Runtime/WebBrowser/WebBrowser.Build.cs
@@ -10,13 +10,14 @@ public class WebBrowser : ModuleRules
 		PrivateDependencyModuleNames.AddRange(
 			new string[]
 			{
+				"Engine",
 				"Core",
 				"CoreUObject",
 				"ApplicationCore",
 				"RHI",
 				"InputCore",
 				"Serialization",
-				"HTTP"
+				"HTTP",
 			}
 		);
 
