[WebBrowser]

Remove blanket "Engine" dependency.
Added specific "Engine" dependency to Linux if target type is not Program.

diff --git a/Engine/Source/Runtime/WebBrowser/WebBrowser.Build.cs b/Engine/Source/Runtime/WebBrowser/WebBrowser.Build.cs
--- a/Engine/Source/Runtime/WebBrowser/WebBrowser.Build.cs
+++ b/Engine/Source/Runtime/WebBrowser/WebBrowser.Build.cs
@@ -10,7 +10,6 @@ public class WebBrowser : ModuleRules
 		PrivateDependencyModuleNames.AddRange(
 			new string[]
 			{
-				"Engine",
 				"Core",
 				"CoreUObject",
 				"ApplicationCore",
@@ -60,7 +59,17 @@ public class WebBrowser : ModuleRules
 
 			CircularlyReferencedDependentModules.Add("WebBrowserTexture");
 		}
-
+		
+		if (Target.Type != TargetType.Program && Target.Platform == UnrealTargetPlatform.Linux)
+		{
+			PrivateDependencyModuleNames.AddRange(
+				new string[]
+				{
+					"Engine"
+				}
+			);
+		}
+		
 		if ( Target.Type != TargetType.Program && Target.Platform == UnrealTargetPlatform.Win64)
         {
 			PrivateDependencyModuleNames.AddRange(
