[WebBrowser]

CapturedCefBuffer member and includes enclosed in #WITH_CEF3.

diff --git a/Engine/Source/Runtime/WebBrowser/Private/CapturedCefBuffer.cpp b/Engine/Source/Runtime/WebBrowser/Private/CapturedCefBuffer.cpp
--- a/Engine/Source/Runtime/WebBrowser/Private/CapturedCefBuffer.cpp
+++ b/Engine/Source/Runtime/WebBrowser/Private/CapturedCefBuffer.cpp
@@ -3,13 +3,19 @@
 
 #include "CapturedCefBuffer.h"
 
+
+#if WITH_CEF3
 #include "Runtime/Engine/Public/TextureResource.h"
 #include "Framework/Application/SlateApplication.h"
 #include "Styling/StyleColors.h"
+#endif
 
 
 bool FCapturedCefBuffer::SetBufferAsB8G8R8A8(const void* InBufferB8G8R8A8, const int32 InBufferWidth, const int32 InBufferHeight, const float InViewportDPIScaleFactor, const bool bDoSkipBadBufferTest)
 {
+#if WITH_CEF3
+
+	
 	if (!InBufferB8G8R8A8 || InBufferWidth <= 0 || InBufferHeight <= 0)
 	{
 		ClearBuffer();
@@ -62,26 +68,43 @@ bool FCapturedCefBuffer::SetBufferAsB8G8R8A8(const void* InBufferB8G8R8A8, const
 	}
 
 	
+#endif
+
+	
 	return false;
 }
 
 
 bool FCapturedCefBuffer::ClearBuffer()
 {
+#if WITH_CEF3
+
+	
 	if (!BufferData.IsEmpty())
 	{
 		BufferData.Empty();
 	}
 	BufferDimensions = FIntPoint::ZeroValue;
 	ViewportDPIScaleFactor = 1.0f;
-
 	
 	return true;
+
+
+#else
+
+
+	return false;
+
+	
+#endif
 }
 
 
 bool FCapturedCefBuffer::MaybeUpdatePaintObjects() 
 {
+#if WITH_CEF3
+
+	
 	if (!FSlateApplication::IsInitialized()) // ..safety check
 	{
 		return false;
@@ -128,11 +151,23 @@ bool FCapturedCefBuffer::MaybeUpdatePaintObjects()
 
 	
 	return true;
+
+	
+#else
+
+	
+	return false;
+	
+	
+#endif
 }
 
 
 bool FCapturedCefBuffer::ClearPaintObjects()
 {
+#if WITH_CEF3
+
+	
 	if (!bHasPaintObjects)
 	{
 		return false;
@@ -158,11 +193,23 @@ bool FCapturedCefBuffer::ClearPaintObjects()
 
 
 	return true;
+
+
+#else
+
+	
+	return false;
+
+
+#endif
 }
 
 
 int32 FCapturedCefBuffer::PaintCentered(const FGeometry& AllottedGeometry, FSlateWindowElementList& OutDrawElements, int32 LayerId, const FWidgetStyle& InWidgetStyle) const
 {
+#if WITH_CEF3
+
+	
 	if (!bHasPaintObjects || ViewportDPIScaleFactor <= 0.0)
 	{
 		return LayerId;
@@ -186,23 +233,51 @@ int32 FCapturedCefBuffer::PaintCentered(const FGeometry& AllottedGeometry, FSlat
 	);
 
 	
+#endif
+	
+	
 	return LayerId;
 }
 
 
 bool FCapturedCefBuffer::HasPaintObjects() const
 {
+#if WITH_CEF3
+	
 	return bHasPaintObjects;
+	
+#else
+	
+	return false;
+	
+#endif
 }
 
 
 bool FCapturedCefBuffer::IsBufferValid() const
 {
-	return (!BufferData.IsEmpty() && BufferDimensions.X > 0 && BufferDimensions.Y > 0);	
+#if WITH_CEF3
+
+	return (!BufferData.IsEmpty() && BufferDimensions.X > 0 && BufferDimensions.Y > 0);
+	
+#else
+
+	return false;
+	
+#endif
 }
 
 
 uint32 FCapturedCefBuffer::CalculateBufferNumBytes() const
 {
+#if WITH_CEF3
+
 	return (BufferDimensions.X * BufferDimensions.Y * 4/*PF_B8G8R8A8*/);
+	
+#else
+
+	return 0;
+	
+#endif
 }
+
diff --git a/Engine/Source/Runtime/WebBrowser/Private/CapturedCefBuffer.h b/Engine/Source/Runtime/WebBrowser/Private/CapturedCefBuffer.h
--- a/Engine/Source/Runtime/WebBrowser/Private/CapturedCefBuffer.h
+++ b/Engine/Source/Runtime/WebBrowser/Private/CapturedCefBuffer.h
@@ -3,12 +3,14 @@
 
 #pragma once
 
+#include "Layout/Geometry.h"
+#include "Rendering/DrawElements.h"
 
+#if WITH_CEF3
 #include "Engine/Texture2D.h"
 #include "UObject/StrongObjectPtrTemplates.h"
-#include "Layout/Geometry.h"
-#include "Rendering/DrawElements.h"
 #include "Styling/SlateBrush.h"
+#endif
 
 #include "CapturedCefBuffer.generated.h"
 
@@ -31,6 +33,7 @@ public:
 
 	bool HasPaintObjects() const;
 
+	
 private:
 
 
@@ -38,6 +41,8 @@ private:
 	uint32 CalculateBufferNumBytes() const;
 	
 
+#if WITH_CEF3
+	
 	TArray<uint8> BufferData;
 	FIntPoint BufferDimensions = FIntPoint::ZeroValue;
 
@@ -47,4 +52,6 @@ private:
 	bool bHasPaintObjects = false;
 	TStrongObjectPtr<UTexture2D> PaintTexture;
 	FSlateBrush PaintSlateBrush;
+	
+#endif
 };
