[WebBrowser]

Fixes to make programs w/o Engine build (i.e. SlateViewer on Win64.)

diff --git a/Engine/Source/Runtime/WebBrowser/Private/CEF/CEFWebBrowserWindow.cpp b/Engine/Source/Runtime/WebBrowser/Private/CEF/CEFWebBrowserWindow.cpp
--- a/Engine/Source/Runtime/WebBrowser/Private/CEF/CEFWebBrowserWindow.cpp
+++ b/Engine/Source/Runtime/WebBrowser/Private/CEF/CEFWebBrowserWindow.cpp
@@ -2928,7 +2928,7 @@ bool FCEFWebBrowserWindow::HasCorrectNativeCefBuffer() const
 
 void FCEFWebBrowserWindow::SetUseCapturedCefBuffer(const bool bValue)
 {
-#if UE_CEF_HAS_RESIZE_BUG
+#if UE_CEF_HAS_RESIZE_BUG && UE_WITH_CAPTURED_CEF_BUFFER
 	if (bValue)
 	{
 		if (!CapturedCefBuffer.IsValid())
diff --git a/Engine/Source/Runtime/WebBrowser/Private/CapturedCefBuffer.cpp b/Engine/Source/Runtime/WebBrowser/Private/CapturedCefBuffer.cpp
--- a/Engine/Source/Runtime/WebBrowser/Private/CapturedCefBuffer.cpp
+++ b/Engine/Source/Runtime/WebBrowser/Private/CapturedCefBuffer.cpp
@@ -4,7 +4,7 @@
 #include "CapturedCefBuffer.h"
 
 
-#if WITH_CEF3
+#if UE_WITH_CAPTURED_CEF_BUFFER
 #include "Runtime/Engine/Public/TextureResource.h"
 #include "Framework/Application/SlateApplication.h"
 #include "Styling/StyleColors.h"
@@ -13,8 +13,8 @@
 
 bool FCapturedCefBuffer::SetBufferAsB8G8R8A8(const void* InBufferB8G8R8A8, const int32 InBufferWidth, const int32 InBufferHeight, const float InViewportDPIScaleFactor, const bool bDoSkipBadBufferTest)
 {
-#if WITH_CEF3
-
+#if UE_WITH_CAPTURED_CEF_BUFFER
+	
 	
 	if (!InBufferB8G8R8A8 || InBufferWidth <= 0 || InBufferHeight <= 0)
 	{
@@ -77,7 +77,7 @@ bool FCapturedCefBuffer::SetBufferAsB8G8R8A8(const void* InBufferB8G8R8A8, const
 
 bool FCapturedCefBuffer::ClearBuffer()
 {
-#if WITH_CEF3
+#if UE_WITH_CAPTURED_CEF_BUFFER
 
 	
 	if (!BufferData.IsEmpty())
@@ -102,8 +102,8 @@ bool FCapturedCefBuffer::ClearBuffer()
 
 bool FCapturedCefBuffer::MaybeUpdatePaintObjects() 
 {
-#if WITH_CEF3
-
+#if UE_WITH_CAPTURED_CEF_BUFFER
+	
 	
 	if (!FSlateApplication::IsInitialized()) // ..safety check
 	{
@@ -165,8 +165,8 @@ bool FCapturedCefBuffer::MaybeUpdatePaintObjects()
 
 bool FCapturedCefBuffer::ClearPaintObjects()
 {
-#if WITH_CEF3
-
+#if UE_WITH_CAPTURED_CEF_BUFFER
+	
 	
 	if (!bHasPaintObjects)
 	{
@@ -207,8 +207,8 @@ bool FCapturedCefBuffer::ClearPaintObjects()
 
 int32 FCapturedCefBuffer::PaintCentered(const FGeometry& AllottedGeometry, FSlateWindowElementList& OutDrawElements, int32 LayerId, const FWidgetStyle& InWidgetStyle) const
 {
-#if WITH_CEF3
-
+#if UE_WITH_CAPTURED_CEF_BUFFER
+	
 	
 	if (!bHasPaintObjects || ViewportDPIScaleFactor <= 0.0)
 	{
@@ -242,7 +242,7 @@ int32 FCapturedCefBuffer::PaintCentered(const FGeometry& AllottedGeometry, FSlat
 
 bool FCapturedCefBuffer::HasPaintObjects() const
 {
-#if WITH_CEF3
+#if UE_WITH_CAPTURED_CEF_BUFFER
 	
 	return bHasPaintObjects;
 	
@@ -254,30 +254,18 @@ bool FCapturedCefBuffer::HasPaintObjects() const
 }
 
 
+#if UE_WITH_CAPTURED_CEF_BUFFER
 bool FCapturedCefBuffer::IsBufferValid() const
 {
-#if WITH_CEF3
-
 	return (!BufferData.IsEmpty() && BufferDimensions.X > 0 && BufferDimensions.Y > 0);
-	
-#else
-
-	return false;
-	
-#endif
 }
+#endif
 
 
+#if UE_WITH_CAPTURED_CEF_BUFFER
 uint32 FCapturedCefBuffer::CalculateBufferNumBytes() const
 {
-#if WITH_CEF3
-
 	return (BufferDimensions.X * BufferDimensions.Y * 4/*PF_B8G8R8A8*/);
-	
-#else
-
-	return 0;
-	
-#endif
 }
+#endif
 
diff --git a/Engine/Source/Runtime/WebBrowser/Private/CapturedCefBuffer.h b/Engine/Source/Runtime/WebBrowser/Private/CapturedCefBuffer.h
--- a/Engine/Source/Runtime/WebBrowser/Private/CapturedCefBuffer.h
+++ b/Engine/Source/Runtime/WebBrowser/Private/CapturedCefBuffer.h
@@ -3,10 +3,14 @@
 
 #pragma once
 
+
+#define UE_WITH_CAPTURED_CEF_BUFFER (WITH_ENGINE && WITH_CEF3)
+
+
 #include "Layout/Geometry.h"
 #include "Rendering/DrawElements.h"
 
-#if WITH_CEF3
+#if UE_WITH_CAPTURED_CEF_BUFFER
 #include "Engine/Texture2D.h"
 #include "UObject/StrongObjectPtrTemplates.h"
 #include "Styling/SlateBrush.h"
@@ -37,12 +41,13 @@ public:
 private:
 
 
+#if UE_WITH_CAPTURED_CEF_BUFFER
+
+	
 	bool IsBufferValid() const; 
 	uint32 CalculateBufferNumBytes() const;
 	
 
-#if WITH_CEF3
-	
 	TArray<uint8> BufferData;
 	FIntPoint BufferDimensions = FIntPoint::ZeroValue;
 
@@ -52,6 +57,7 @@ private:
 	bool bHasPaintObjects = false;
 	TStrongObjectPtr<UTexture2D> PaintTexture;
 	FSlateBrush PaintSlateBrush;
+
 	
 #endif
 };
