[WebBrowser]

Build.cs -
For Win64 and Linux..
Engine module is a dependency is added if Target.bCompileAgainstEngine is true, not just when Target.Type!=TargetType.Program.
This makes Chaos Visual Debugger build, and will work for any Win64 or Linux build that needs the Editor module - including non-program or program targets,

Texture2d.h include moved out of CapturedCefBuffer.h..
But this wasn't strictly necessary for this fix.

diff --git a/Engine/Source/Runtime/WebBrowser/Private/CapturedCefBuffer.cpp b/Engine/Source/Runtime/WebBrowser/Private/CapturedCefBuffer.cpp
--- a/Engine/Source/Runtime/WebBrowser/Private/CapturedCefBuffer.cpp
+++ b/Engine/Source/Runtime/WebBrowser/Private/CapturedCefBuffer.cpp
@@ -5,6 +5,7 @@
 
 
 #if UE_WITH_CAPTURED_CEF_BUFFER
+#include "Engine/Texture2D.h"
 #include "Runtime/Engine/Public/TextureResource.h"
 #include "Framework/Application/SlateApplication.h"
 #include "Styling/StyleColors.h"
diff --git a/Engine/Source/Runtime/WebBrowser/Private/CapturedCefBuffer.h b/Engine/Source/Runtime/WebBrowser/Private/CapturedCefBuffer.h
--- a/Engine/Source/Runtime/WebBrowser/Private/CapturedCefBuffer.h
+++ b/Engine/Source/Runtime/WebBrowser/Private/CapturedCefBuffer.h
@@ -11,9 +11,9 @@
 #include "Rendering/DrawElements.h"
 
 #if UE_WITH_CAPTURED_CEF_BUFFER
-#include "Engine/Texture2D.h"
 #include "UObject/StrongObjectPtrTemplates.h"
 #include "Styling/SlateBrush.h"
+class UTexture2D;
 #endif
 
 #include "CapturedCefBuffer.generated.h"
diff --git a/Engine/Source/Runtime/WebBrowser/WebBrowser.Build.cs b/Engine/Source/Runtime/WebBrowser/WebBrowser.Build.cs
--- a/Engine/Source/Runtime/WebBrowser/WebBrowser.Build.cs
+++ b/Engine/Source/Runtime/WebBrowser/WebBrowser.Build.cs
@@ -60,7 +60,7 @@ public class WebBrowser : ModuleRules
 			CircularlyReferencedDependentModules.Add("WebBrowserTexture");
 		}
 		
-		if (Target.Type != TargetType.Program && Target.Platform == UnrealTargetPlatform.Linux)
+		if (Target.bCompileAgainstEngine && Target.Platform == UnrealTargetPlatform.Linux)
 		{
 			PrivateDependencyModuleNames.AddRange(
 				new string[]
@@ -70,7 +70,7 @@ public class WebBrowser : ModuleRules
 			);
 		}
 		
-		if ( Target.Type != TargetType.Program && Target.Platform == UnrealTargetPlatform.Win64)
+		if (Target.bCompileAgainstEngine && Target.Platform == UnrealTargetPlatform.Win64)
         {
 			PrivateDependencyModuleNames.AddRange(
 				new string[]
@@ -79,7 +79,6 @@ public class WebBrowser : ModuleRules
 					"RenderCore"
 				}
 			);
-
 		}
 
 		if (Target.Platform == UnrealTargetPlatform.Android)
