WebBrowser: Use the GL backend for ANGLE on Linux

The Vulkan backend fails to initialize, likely due to the removal of the conflicting libvulkan.so.1 loader at CL43202912.
The GL backend is still completely supported and provides a more stable environment, so stick to it for now.

#rnx

diff --git a/Engine/Source/Runtime/WebBrowser/Private/CEF/CEFBrowserApp.cpp b/Engine/Source/Runtime/WebBrowser/Private/CEF/CEFBrowserApp.cpp
--- a/Engine/Source/Runtime/WebBrowser/Private/CEF/CEFBrowserApp.cpp
+++ b/Engine/Source/Runtime/WebBrowser/Private/CEF/CEFBrowserApp.cpp
@@ -43,9 +43,17 @@ void FCEFBrowserApp::OnBeforeCommandLineProcessing(const CefString& ProcessType,
 
 #if PLATFORM_LINUX
 	CommandLine->AppendSwitchWithValue("ozone-platform", "headless");
+	// Stick to the GL backend for Angle as the Vulkan backend is failing to initialize,
+	// likely due to the removal of the conflicting libvulkan.so.1 loader at CL43202912
+#if 0
 	CommandLine->AppendSwitchWithValue("use-gl", "angle");
 	CommandLine->AppendSwitchWithValue("use-angle", "vulkan");
-	CommandLine->AppendSwitch("use-vulkan");
+	CommandLine->AppendSwitchWithValue("use-vulkan", "native");
+	CommandLine->AppendSwitchWithValue("enable-features", "Vulkan");
+#else
+	CommandLine->AppendSwitchWithValue("use-gl", "angle");
+	CommandLine->AppendSwitchWithValue("use-angle", "gl");
+#endif
 #endif
 
 	CommandLine->AppendSwitch("enable-begin-frame-scheduling");
