[WebBrowser]

CapturedCefBuffer -
Moved include for Texture2D back into header, since TStrongObjectPtr<UTexture2D> cannot use incomplete type.

diff --git a/Engine/Source/Runtime/WebBrowser/Private/CapturedCefBuffer.cpp b/Engine/Source/Runtime/WebBrowser/Private/CapturedCefBuffer.cpp
--- a/Engine/Source/Runtime/WebBrowser/Private/CapturedCefBuffer.cpp
+++ b/Engine/Source/Runtime/WebBrowser/Private/CapturedCefBuffer.cpp
@@ -5,7 +5,6 @@
 
 
 #if UE_WITH_CAPTURED_CEF_BUFFER
-#include "Engine/Texture2D.h"
 #include "Runtime/Engine/Public/TextureResource.h"
 #include "Framework/Application/SlateApplication.h"
 #include "Styling/StyleColors.h"
diff --git a/Engine/Source/Runtime/WebBrowser/Private/CapturedCefBuffer.h b/Engine/Source/Runtime/WebBrowser/Private/CapturedCefBuffer.h
--- a/Engine/Source/Runtime/WebBrowser/Private/CapturedCefBuffer.h
+++ b/Engine/Source/Runtime/WebBrowser/Private/CapturedCefBuffer.h
@@ -11,9 +11,9 @@
 #include "Rendering/DrawElements.h"
 
 #if UE_WITH_CAPTURED_CEF_BUFFER
+#include "Engine/Texture2D.h"
 #include "UObject/StrongObjectPtrTemplates.h"
 #include "Styling/SlateBrush.h"
-class UTexture2D;
 #endif
 
 #include "CapturedCefBuffer.generated.h"
