Fixing Shown vs Hidden Tab Visibility bug with macOS WebKit and the Fab and Quixel plugins

diff --git a/Engine/Plugins/Bridge/Source/Bridge/Private/UI/BridgeUIManager.cpp b/Engine/Plugins/Bridge/Source/Bridge/Private/UI/BridgeUIManager.cpp
--- a/Engine/Plugins/Bridge/Source/Bridge/Private/UI/BridgeUIManager.cpp
+++ b/Engine/Plugins/Bridge/Source/Bridge/Private/UI/BridgeUIManager.cpp
@@ -339,6 +339,7 @@ TSharedRef<SDockTab> FBridgeUIManagerImpl::CreateBridgeTab(const FSpawnTabArgs&
 			FBridgeUIManager::Instance->LocalBrowserDock = NULL;
 			if (FBridgeUIManager::Instance->WebBrowserWidget.IsValid())
 			{
+				FBridgeUIManager::Instance->Browser->SetParentDockTab(nullptr);
 				FBridgeUIManager::Instance->WebBrowserWidget.Reset();
 				FBridgeUIManager::Instance->Browser.Reset();
 			}
@@ -371,10 +372,15 @@ TSharedRef<SDockTab> FBridgeUIManagerImpl::CreateBridgeTab(const FSpawnTabArgs&
 
 	if (LocalBrowserDock.IsValid())
 	{
+		Browser->SetParentDockTab(LocalBrowserDock);
 		return LocalBrowserDock.ToSharedRef();
 	}
+	
+	SAssignNew(LocalBrowserDock, SDockTab);
+	
+	Browser->SetParentDockTab(LocalBrowserDock);
 
-	return SAssignNew(LocalBrowserDock, SDockTab);
+	return LocalBrowserDock.ToSharedRef();
 }
 
 #undef LOCTEXT_NAMESPACE
diff --git a/Engine/Plugins/Fab/Source/Fab/Private/FabBrowser.cpp b/Engine/Plugins/Fab/Source/Fab/Private/FabBrowser.cpp
--- a/Engine/Plugins/Fab/Source/Fab/Private/FabBrowser.cpp
+++ b/Engine/Plugins/Fab/Source/Fab/Private/FabBrowser.cpp
@@ -367,6 +367,8 @@ TSharedRef<SDockTab> FFabBrowser::OpenTab(const FSpawnTabArgs&)
 	[
 		WebBrowserInstance.ToSharedRef()
 	];
+	
+	WebBrowserWindow->SetParentDockTab(DockTab);
 
 	return DockTab.ToSharedRef();
 }
@@ -387,6 +389,11 @@ void FFabBrowser::OnPluginTabClosed(TSharedRef<SDockTab> InParentTab)
 		}
 	);
 
+#if PLATFORM_MAC && !WITH_CEF3
+	WebBrowserWindow->SetIsDisabled(true);
+	WebBrowserWindow->SetParentDockTab(nullptr);
+#endif
+	
 	AsyncTask(ENamedThreads::AnyBackgroundThreadNormalTask, []()
 	{
 		// add a small delay so that the frontend can register the close event
diff --git a/Engine/Source/Runtime/WebBrowser/Private/Apple/ApplePlatformWebBrowser.cpp b/Engine/Source/Runtime/WebBrowser/Private/Apple/ApplePlatformWebBrowser.cpp
--- a/Engine/Source/Runtime/WebBrowser/Private/Apple/ApplePlatformWebBrowser.cpp
+++ b/Engine/Source/Runtime/WebBrowser/Private/Apple/ApplePlatformWebBrowser.cpp
@@ -20,6 +20,7 @@
 #include "HAL/PlatformProcess.h"
 
 #include "Framework/Application/SlateApplication.h"
+#include "Framework/Docking/TabManager.h"
 #include "Rendering/SlateRenderer.h"
 #include "Textures/SlateUpdatableTexture.h"
 
@@ -645,6 +646,28 @@ class SAppleWebBrowserWidget : public SViewport
 		return WebViewWrapper != nil && [WebViewWrapper inKeyWindow];
 	}
 	
+	void OnActiveTabChanged(TSharedPtr<SDockTab> ActiveTab, TSharedPtr<SDockTab> PreviouslyActiveTab)
+	{
+		TSharedPtr<FWebBrowserWindow> WebBrowserWindowStrongPtr = WebBrowserWindowPtr.Pin();
+		
+		if(WebBrowserWindowStrongPtr.IsValid())
+		{
+			TSharedPtr<SDockTab> CurrentTab = WebBrowserWindowStrongPtr->GetParentDockTab().Pin();
+			
+			// this should only be called when the active browser has an active parent dock tab
+			check(CurrentTab.IsValid());
+			
+			if(ActiveTab == CurrentTab)
+			{
+				[WebViewWrapper setVisibility: YES ];
+			}
+			else if(PreviouslyActiveTab == CurrentTab)
+			{
+				[WebViewWrapper setVisibility: NO ];
+			}
+		}
+	}
+	
 	~SAppleWebBrowserWidget()
 	{
 		Close();
@@ -1466,6 +1489,14 @@ TSharedRef<SViewport> FWebBrowserWindow::CreateWidget()
 	return BrowserWidgetRef;
 }
 
+void FWebBrowserWindow::OnActiveTabChanged(TSharedPtr<SDockTab> ActiveTab, TSharedPtr<SDockTab> PreviouslyActiveTab)
+{
+	if(BrowserWidget != nullptr)
+	{
+		BrowserWidget->OnActiveTabChanged(ActiveTab, PreviouslyActiveTab);
+	}
+}
+
 void FWebBrowserWindow::SetViewportSize(FIntPoint ViewportSize, FIntPoint ViewportPos)
 {
 	WindowSize = ViewportSize;
@@ -1762,6 +1793,38 @@ void FWebBrowserWindow::SetParentWindow(TSharedPtr<SWindow> Window)
 	ParentWindow = Window;
 }
 
+void FWebBrowserWindow::SetParentDockTab(TSharedPtr<SDockTab> DockTab)
+{
+	if(DockTab == nullptr && ParentDockTabDelegate.IsValid())
+	{
+		FGlobalTabmanager::Get()->OnTabForegrounded_Unsubscribe(ParentDockTabDelegate);
+	}
+	
+	ParentDockTab = DockTab;
+	
+	if(ParentDockTab != nullptr)
+	{
+		FOnActiveTabChanged::FDelegate ActiveTabChangedDelegate;
+		
+		if(!ParentDockTabDelegate.IsValid())
+		{
+			ActiveTabChangedDelegate.BindSP(AsShared(), &FWebBrowserWindow::OnActiveTabChanged);
+			
+			ParentDockTabDelegate = FGlobalTabmanager::Get()->OnTabForegrounded_Subscribe(ActiveTabChangedDelegate);
+		}
+	}
+	else
+	{
+		// if the browser is disabled and we are setting this dock tab to nullptr, we are closing the tab
+		BrowserWidget->SetWebBrowserVisibility(!bIsDisabled);
+	}
+}
+
+TWeakPtr<SDockTab> FWebBrowserWindow::GetParentDockTab()
+{
+	return ParentDockTab;
+}
+
 void FWebBrowserWindow::ShowFloatingCloseButton(bool bShow, bool bDraggable)
 {
 	if (BrowserWidget)
diff --git a/Engine/Source/Runtime/WebBrowser/Private/Apple/ApplePlatformWebBrowser.h b/Engine/Source/Runtime/WebBrowser/Private/Apple/ApplePlatformWebBrowser.h
--- a/Engine/Source/Runtime/WebBrowser/Private/Apple/ApplePlatformWebBrowser.h
+++ b/Engine/Source/Runtime/WebBrowser/Private/Apple/ApplePlatformWebBrowser.h
@@ -181,6 +181,7 @@ public:
 	virtual void SetIsDisabled(bool bValue) override;
 	virtual TSharedPtr<SWindow> GetParentWindow() const override;
 	virtual void SetParentWindow(TSharedPtr<SWindow> Window) override;
+	virtual void SetParentDockTab(TSharedPtr<SDockTab> DockTab) override;
 	virtual void ShowFloatingCloseButton(bool bShow, bool bDraggable) override;
 
 	// TODO: None of these events are actually called
@@ -349,6 +350,10 @@ public:
 		Title = InTitle;
 		OnTitleChanged().Broadcast(Title);
 	}
+	
+	void OnActiveTabChanged(TSharedPtr<SDockTab> ActiveTab, TSharedPtr<SDockTab> PreviouslyActiveTab);
+	
+	TWeakPtr<SDockTab> GetParentDockTab();
 
 private:
 
@@ -452,6 +457,8 @@ private:
 	FOnUnhandledKeyChar UnhandledKeyCharDelegate;
 
 	TSharedPtr<SWindow> ParentWindow;
+	TWeakPtr<SDockTab> ParentDockTab;
+	FDelegateHandle ParentDockTabDelegate;
 
 	FIntPoint WindowSize;
 
diff --git a/Engine/Source/Runtime/WebBrowser/Public/IWebBrowserWindow.h b/Engine/Source/Runtime/WebBrowser/Public/IWebBrowserWindow.h
--- a/Engine/Source/Runtime/WebBrowser/Public/IWebBrowserWindow.h
+++ b/Engine/Source/Runtime/WebBrowser/Public/IWebBrowserWindow.h
@@ -374,6 +374,8 @@ public:
 	virtual void ShowFloatingCloseButton(bool bShow, bool bDraggable) = 0;
 
 	virtual void CheckTickActivity() {};
+	
+	virtual void SetParentDockTab(TSharedPtr<class SDockTab> DockTab) {};
 
 public:
 
