[WebBrowser]

For IME with languages like "Chinese (traditional, Taiwan)"..
Set 'underlines' to have no appearance, to avoid default yellow highlighting that makes the text unreadable and unlike other IME input UX.

diff --git a/Engine/Source/Runtime/WebBrowser/Private/CEF/CEFTextInputMethodContext.cpp b/Engine/Source/Runtime/WebBrowser/Private/CEF/CEFTextInputMethodContext.cpp
--- a/Engine/Source/Runtime/WebBrowser/Private/CEF/CEFTextInputMethodContext.cpp
+++ b/Engine/Source/Runtime/WebBrowser/Private/CEF/CEFTextInputMethodContext.cpp
@@ -88,6 +88,19 @@ void FCEFTextInputMethodContext::SetSelectionRange(const uint32 BeginIndex, cons
 
 	CefString Str = TCHAR_TO_WCHAR(*CompositionString);
 	std::vector<CefCompositionUnderline> underlines;
+
+	// Prevent CEF's ugly/unreadable default highlighting.
+	CefCompositionUnderline Underline;
+	{
+		Underline.range.from = 0;
+		Underline.range.to = Str.length();
+		Underline.color = 0;
+		Underline.background_color = 0;
+		Underline.thick = 0;
+		Underline.style = CEF_CUS_NONE;
+	}
+	underlines.push_back(Underline);
+	
 	Owner->InternalCefBrowser->GetHost()->ImeSetComposition(
 		Str,
 		underlines,
@@ -118,6 +131,19 @@ void FCEFTextInputMethodContext::SetTextInRange(const uint32 BeginIndex, const u
 
 	CefString Str = TCHAR_TO_WCHAR(*CompositionString);
 	std::vector<CefCompositionUnderline> underlines;
+	
+	// Prevent CEF's ugly/unreadable default highlighting.
+	CefCompositionUnderline Underline;
+	{
+		Underline.range.from = 0;
+		Underline.range.to = Str.length();
+		Underline.color = 0;
+		Underline.background_color = 0;
+		Underline.thick = 0;
+		Underline.style = CEF_CUS_NONE;
+	}
+	underlines.push_back(Underline);
+
 	Owner->InternalCefBrowser->GetHost()->ImeSetComposition(
 		Str,
 		underlines,
