WebBrowser: enable accelerated painting on Win64 for D3D11 and D3D12 RHI backends.

Only enable for the new CEF v128 though, and get rid of the old implementation which was disabled anyway and only supported the outdated D3D11 RHI backend.

Also apply some of the resizing fixes from the non-accelerated path to make resizing smoother.

#rnx

diff --git a/Engine/Source/Runtime/WebBrowser/Private/CEF/CEFWebBrowserWindow.cpp b/Engine/Source/Runtime/WebBrowser/Private/CEF/CEFWebBrowserWindow.cpp
--- a/Engine/Source/Runtime/WebBrowser/Private/CEF/CEFWebBrowserWindow.cpp
+++ b/Engine/Source/Runtime/WebBrowser/Private/CEF/CEFWebBrowserWindow.cpp
@@ -427,8 +427,8 @@ FCEFWebBrowserWindow::FCEFWebBrowserWindow(CefRefPtr<CefBrowser> InBrowser, CefR
 	check(InBrowser.get() != nullptr);
 	check(!bUsingAcceleratedPaint || CanSupportAcceleratedPaint()); // make sure if accelerated paint is selected we can support it
 
-	UpdatableTextures[0] = nullptr;
-	UpdatableTextures[1] = nullptr;
+	UpdatableTextures[PET_VIEW] = nullptr;
+	UpdatableTextures[PET_POPUP] = nullptr;
 
 	if (!CreateInitialTextures())
 	{
@@ -449,17 +449,21 @@ FCEFWebBrowserWindow::FCEFWebBrowserWindow(CefRefPtr<CefBrowser> InBrowser, CefR
 
 void FCEFWebBrowserWindow::ReleaseTextures()
 {
-	for (int I = 0; I < 2; ++I)
+	for (int Type : { PET_VIEW, PET_POPUP })
 	{
-		if (UpdatableTextures[I] != nullptr)
+		if (UpdatableTextures[Type] != nullptr)
 		{
-			FSlateUpdatableTexture* TextureToRelease = UpdatableTextures[I];
+			FSlateUpdatableTexture* TextureToRelease = UpdatableTextures[Type];
 
 			if (IsInGameThread())
 			{
 				if (FSlateApplication::IsInitialized())
 				{
-					if (FSlateRenderer* Renderer = FSlateApplication::Get().GetRenderer())
+					if (RHIRenderHelper)
+					{
+						RHIRenderHelper->ReleaseSlateUpdatableTexture(TextureToRelease);
+					}
+					else if (FSlateRenderer* Renderer = FSlateApplication::Get().GetRenderer())
 					{
 						Renderer->ReleaseUpdatableTexture(TextureToRelease);
 					}
@@ -467,11 +471,15 @@ void FCEFWebBrowserWindow::ReleaseTextures()
 			}
 			else if (FTaskGraphInterface::IsRunning())
 			{
-				AsyncTask(ENamedThreads::GameThread, [TextureToRelease]()
+				AsyncTask(ENamedThreads::GameThread, [this, TextureToRelease]()
 				{
 					if (FSlateApplication::IsInitialized())
 					{
-						if (FSlateRenderer* Renderer = FSlateApplication::Get().GetRenderer())
+						if (RHIRenderHelper)
+						{
+							RHIRenderHelper->ReleaseSlateUpdatableTexture(TextureToRelease);
+						}
+						else if (FSlateRenderer* Renderer = FSlateApplication::Get().GetRenderer())
 						{
 							Renderer->ReleaseUpdatableTexture(TextureToRelease);
 						}
@@ -479,7 +487,7 @@ void FCEFWebBrowserWindow::ReleaseTextures()
 				});
 			}
 
-			UpdatableTextures[I] = nullptr;
+			UpdatableTextures[Type] = nullptr;
 		}
 	}
 }
@@ -497,43 +505,46 @@ bool FCEFWebBrowserWindow::CreateInitialTextures()
 
 			if (bUsingAcceleratedPaint)
 			{
-				if (FCEFWebBrowserWindowRHIHelper::BUseRHIRenderer() && RHIRenderHelper == nullptr)
+				if (FCEFWebBrowserWindowRHIHelper::BUseSupportedRHIRenderer() && !RHIRenderHelper)
 				{
-					RHIRenderHelper = new FCEFWebBrowserWindowRHIHelper;
+					RHIRenderHelper.Reset(new FCEFWebBrowserWindowRHIHelper);
+
+					// Render a transparent texture until we receive the first update from CEF
+					UpdatableTextures[PET_VIEW] = RHIRenderHelper->CreateSlateUpdatableTexture({ 1, 1 });
+					UpdatableTextures[PET_POPUP] = RHIRenderHelper->CreateSlateUpdatableTexture({ 1, 1 });
+
+					return true;
 				}
-				
-				// the accelerated paint path attaches to the texture at render time as we don't know its details until then
-				UpdatableTextures[0] = nullptr;
-				UpdatableTextures[1] = nullptr;
-				return true;
+
+				return false;
 			}
 
 			// Create a transparent dummy texture for our buffers which will prevent slate from applying an 
 			// undesirable quad if it happens to ask for this buffer before we get a chance to paint to it.
 			TArray<uint8> RawData;
 			RawData.AddZeroed(4);
-			UpdatableTextures[0] = Renderer->CreateUpdatableTexture(1, 1);
+			UpdatableTextures[PET_VIEW] = Renderer->CreateUpdatableTexture(1, 1);
 
 			if (Renderer->HasLostDevice())
 			{
 				return false;
 			}
 
-			UpdatableTextures[0]->UpdateTextureThreadSafeRaw(1, 1, RawData.GetData());
+			UpdatableTextures[PET_VIEW]->UpdateTextureThreadSafeRaw(1, 1, RawData.GetData());
 
 			if (Renderer->HasLostDevice())
 			{
 				return false;
 			}
 
-			UpdatableTextures[1] = Renderer->CreateUpdatableTexture(1, 1);
+			UpdatableTextures[PET_POPUP] = Renderer->CreateUpdatableTexture(1, 1);
 
 			if (Renderer->HasLostDevice())
 			{
 				return false;
 			}
 
-			UpdatableTextures[1]->UpdateTextureThreadSafeRaw(1, 1, RawData.GetData());
+			UpdatableTextures[PET_POPUP]->UpdateTextureThreadSafeRaw(1, 1, RawData.GetData());
 
 			if (Renderer->HasLostDevice())
 			{
@@ -563,10 +574,6 @@ FCEFWebBrowserWindow::~FCEFWebBrowserWindow()
 	ReleaseTextures();
 
 	BufferedVideo.Reset();
-	if (RHIRenderHelper != nullptr)
-	{
-		delete RHIRenderHelper;
-	}
 
 	UE_LOG(LogWebBrowser, Log, TEXT("Deleting browser for Url=%s."), *CurrentUrl);
 }
@@ -587,8 +594,7 @@ TSharedRef<SViewport> FCEFWebBrowserWindow::CreateWidget()
 		SNew(SViewport)
 		.EnableGammaCorrection(false)
 		.EnableBlending(bUseTransparency)
-		.IgnoreTextureAlpha(!bUseTransparency)
-		.RenderTransform(this, &FCEFWebBrowserWindow::GetWebBrowserRenderTransform);
+		.IgnoreTextureAlpha(!bUseTransparency);
 
 #if !PLATFORM_LINUX
 	Ime->CacheBrowserSlateInfo(BrowserWidgetRef);
@@ -597,23 +603,6 @@ TSharedRef<SViewport> FCEFWebBrowserWindow::CreateWidget()
 	return BrowserWidgetRef;
 }
 
-TOptional<FSlateRenderTransform> FCEFWebBrowserWindow::GetWebBrowserRenderTransform() const
-{
-	TOptional<FSlateRenderTransform> LocalRenderTransform = FSlateRenderTransform();
-	if (bUsingAcceleratedPaint)
-	{
-		if (RHIRenderHelper != nullptr)
-		{
-			LocalRenderTransform = RHIRenderHelper->GetWebBrowserRenderTransform();
-		}
-		else
-		{
-			LocalRenderTransform = FSlateRenderTransform(Concatenate(FScale2D(1, -1), FVector2D(0, ViewportSize.Y)));
-		}
-	}
-	return LocalRenderTransform;
-}
-
 bool FCEFWebBrowserWindow::BlockInputInDirectHwndMode() const
 {
 #if PLATFORM_WINDOWS
@@ -700,9 +689,9 @@ void FCEFWebBrowserWindow::SetViewportSize(FIntPoint WindowSize, FIntPoint Windo
 
 FSlateShaderResource* FCEFWebBrowserWindow::GetTexture(bool bIsPopup)
 {
-	if (UpdatableTextures[bIsPopup?1:0] != nullptr)
+	if (UpdatableTextures[bIsPopup? PET_POPUP : PET_VIEW] != nullptr)
 	{
-		return UpdatableTextures[bIsPopup?1:0]->GetSlateResource();
+		return UpdatableTextures[bIsPopup? PET_POPUP : PET_VIEW]->GetSlateResource();
 	}
 	return nullptr;
 }
@@ -1986,7 +1975,7 @@ void FCEFWebBrowserWindow::OnPaint(CefRenderHandler::PaintElementType Type, cons
 #if UE_CEF_HAS_RESIZE_BUG
 	const FIntPoint PaintedBufferSize(Width / ViewportDPIScaleFactor, Height / ViewportDPIScaleFactor);
 	bHasCorrectNativeCefBuffer = (PaintedBufferSize == ViewportSize);
-	if (IsUsingCapturedCefBuffer()) 
+	if (CapturedCefBuffer.IsValid())
 	{
 		// We attempt to capture a good CEF buffer, especially when it is correct. The CEF buffer size does not actually need to match the viewport
 		// size to end up being captured, but also may not be captured if it fails another test (which we skip if we know we have the correct CEF
@@ -2004,10 +1993,9 @@ void FCEFWebBrowserWindow::OnPaint(CefRenderHandler::PaintElementType Type, cons
 		bUsingAcceleratedPaint = false;
 		if (UpdatableTextures[Type] != nullptr)
 		{
-			if (FSlateRenderer* const Renderer = GetRenderer())
+			if (RHIRenderHelper)
 			{
-				Renderer->ReleaseUpdatableTexture(UpdatableTextures[Type]);
-				HandleRenderingError();
+				RHIRenderHelper->ReleaseSlateUpdatableTexture(UpdatableTextures[Type]);
 				UpdatableTextures[Type] = nullptr;
 			}
 		}
@@ -2040,7 +2028,7 @@ void FCEFWebBrowserWindow::OnPaint(CefRenderHandler::PaintElementType Type, cons
 			// When resizing the view, the last frame we receive occasionally still has a previous viewport size,
 			// which results in contents looking stretched or with black borders.
 #if UE_CEF_HAS_RESIZE_BUG
-			if (!IsUsingCapturedCefBuffer() && Type == PET_VIEW && !bHasCorrectNativeCefBuffer)
+			if (!CapturedCefBuffer.IsValid() && Type == PET_VIEW && !bHasCorrectNativeCefBuffer)
 			{
 				// Note that calling Invalidate() doesn't work here, and there's no other CEF API to trigger a repaint.
 				InternalCefBrowser->GetHost()->WasHidden(true);
@@ -2115,38 +2103,20 @@ void FCEFWebBrowserWindow::OnAcceleratedPaint(CefRenderHandler::PaintElementType
 		return;
 	LastPaintedSharedHandle = SharedHandle;
 #endif
-	FIntRect Dirty = (DirtyRects.size() == 1) ? FIntRect(DirtyRects[0].x, DirtyRects[0].y, DirtyRects[0].x + DirtyRects[0].width, DirtyRects[0].y + DirtyRects[0].height) : FIntRect();
-	if (UpdatableTextures[Type] == nullptr)
+
+	FIntRect DirtyRect = (DirtyRects.size() == 1) ? FIntRect(DirtyRects[0].x, DirtyRects[0].y, DirtyRects[0].x + DirtyRects[0].width, DirtyRects[0].y + DirtyRects[0].height) : FIntRect();
+	FIntRect ViewportRect = FIntRect(0, 0, ViewportSize.X, ViewportSize.Y).Scale(ViewportDPIScaleFactor);
+	bHasCorrectNativeCefBuffer = (DirtyRect == ViewportRect);
+
+	// During resizing, CEF occasionally sends textures with black bars either above/below or left/right and distorted content in between, we want to ignore these textures.
+	// Fortunately, the DirtyRect for these textures (and only these) has a non-zero X/Y depending on the black bars orientation, so we can detect them without inspecting the content.
+	if (DirtyRect.Min.X > 0 || DirtyRect.Min.Y > 0)
 	{
-		if (FSlateRenderer* const Renderer = GetRenderer())
-		{
-			if (RHIRenderHelper)
-			{
-				UpdatableTextures[Type] = RHIRenderHelper->CreateTexture(SharedHandle);
-			}
-			else
-			{
-				UpdatableTextures[Type] = Renderer->CreateSharedHandleTexture(SharedHandle);
-			}
-			Dirty = FIntRect(); // force a fully copy when we make a new texture
-			HandleRenderingError();
-		}
+		return;
 	}
-	if (UpdatableTextures[Type] != nullptr)
-	{
-#if PLATFORM_WINDOWS
-		if (RHIRenderHelper)
-		{
-			RHIRenderHelper->UpdateSharedHandleTexture(SharedHandle, UpdatableTextures[Type], Dirty.Scale(ViewportDPIScaleFactor));
-		}
-		else
-		{
-			UpdatableTextures[Type]->UpdateTextureThreadSafeWithKeyedTextureHandle(SharedHandle, 1, 0, Dirty.Scale(ViewportDPIScaleFactor));
-		}
-#else
-		UpdatableTextures[Type]->UpdateTextureThreadSafeWithKeyedTextureHandle(SharedHandle, 1, 0, Dirty.Scale(ViewportDPIScaleFactor));
-#endif
 
+	if (RHIRenderHelper && RHIRenderHelper->CopySharedTextureSync(UpdatableTextures[Type], SharedHandle, DirtyRect))
+	{
 		bNeedsRedraw = true;
 		if (Type == PET_POPUP && bShowPopupRequested)
 		{
@@ -2159,7 +2129,6 @@ void FCEFWebBrowserWindow::OnAcceleratedPaint(CefRenderHandler::PaintElementType
 			FIntRect PopupRect = FIntRect(PopupPosition, PopupPosition + PopupSize);
 			OnShowPopup().Broadcast(PopupRect);
 		}
-
 	}
 
 	bIsInitialized = true;
@@ -2584,21 +2553,14 @@ CefMouseEvent FCEFWebBrowserWindow::GetCefMouseEvent(const FGeometry& MyGeometry
 {
 	CefMouseEvent Event;
 
-	FGeometry MouseGeometry = MyGeometry;
-	if (bUsingAcceleratedPaint)
-	{
-		// undo the texture flip if we are using accelerated rendering
-		MouseGeometry = MyGeometry.MakeChild(FSlateRenderTransform(FScale2D(1, -1)));
-	}
-
-	float DPIScale = MouseGeometry.Scale;
+	float DPIScale = MyGeometry.Scale;
 	if (TSharedPtr<SWindow> ParentWindowPtr = ParentWindow.Pin())
 	{
 		DPIScale /= ParentWindowPtr->GetNativeWindow()->GetDPIScaleFactor();
 	}
 
 
-	FVector2D LocalPos = MouseGeometry.AbsoluteToLocal(MouseEvent.GetScreenSpacePosition()) * DPIScale;
+	FVector2D LocalPos = MyGeometry.AbsoluteToLocal(MouseEvent.GetScreenSpacePosition()) * DPIScale;
 	if (bIsPopup)
 	{
 		LocalPos += PopupPosition;
@@ -2655,35 +2617,11 @@ bool FCEFWebBrowserWindow::CanSupportAcceleratedPaint()
 		return false;
 	}
 
-	static bool ForceAcceleratedPaint = FParse::Param(FCommandLine::Get(), TEXT("forcecefaccelpaint"));
-	if (ForceAcceleratedPaint)
-	{
-		return true;
-	}
-
 	// Use off screen rendering so we can integrate with our windows
 #if PLATFORM_LINUX
 	return false;
 #elif PLATFORM_WINDOWS
-#if PLATFORM_64BITS
-	return false;
-	/*
-	static bool Windows10OrAbove = FWindowsPlatformMisc::VerifyWindowsVersion(10, 0); //Win10
-	if (Windows10OrAbove == false)
-	{
-		return  -false;
-	}
-
-	// match the logic in GetStandardStandaloneRenderer() from StandaloneRenderer.cpp to check for the OGL slate renderer
-	if (FParse::Param(FCommandLine::Get(), TEXT("opengl")))
-	{
-		return false;
-	}
-	return true;*/
-
-#else
-	return false; // 32-bit windows doesn't have the accelerated rendering patches applied, it can be done if needed
-#endif
+	return FCEFWebBrowserWindowRHIHelper::BUseSupportedRHIRenderer();
 #elif PLATFORM_MAC
 	return false; // Needs RHI support for the CreateSharedHandleTexture call
 #else
@@ -2698,10 +2636,6 @@ void FCEFWebBrowserWindow::UpdateCachedGeometry(const FGeometry& AllottedGeometr
 	// Forward along the geometry for use by IME
 	Ime->UpdateCachedGeometry(AllottedGeometry);
 #endif
-	if (RHIRenderHelper)
-	{
-		RHIRenderHelper->UpdateCachedGeometry(AllottedGeometry);
-	}
 }
 
 void FCEFWebBrowserWindow::CheckTickActivity()
@@ -2945,10 +2879,10 @@ void FCEFWebBrowserWindow::SetUseCapturedCefBuffer(const bool bValue)
 #endif
 }
 
-bool FCEFWebBrowserWindow::IsUsingCapturedCefBuffer() const
+bool FCEFWebBrowserWindow::HasResizeBug() const
 {
 #if UE_CEF_HAS_RESIZE_BUG
-	return CapturedCefBuffer.IsValid();
+	return true;
 #else
 	return false;
 #endif
diff --git a/Engine/Source/Runtime/WebBrowser/Private/CEF/CEFWebBrowserWindow.h b/Engine/Source/Runtime/WebBrowser/Private/CEF/CEFWebBrowserWindow.h
--- a/Engine/Source/Runtime/WebBrowser/Private/CEF/CEFWebBrowserWindow.h
+++ b/Engine/Source/Runtime/WebBrowser/Private/CEF/CEFWebBrowserWindow.h
@@ -588,9 +588,9 @@ public:
 	WEBBROWSER_API void SetUseCapturedCefBuffer(const bool bValue);
 	
 	/**
-	 * Whether we are using the captured CEF buffer.
+	 * Whether the CEF version we're using has a resize bug that sometimes requires requesting new frames after a resize.
 	 */
-	WEBBROWSER_API bool IsUsingCapturedCefBuffer() const;
+	WEBBROWSER_API bool HasResizeBug() const;
 	
 	/**
 	 * Gets an undistorted buffer captured from CEF. Format is PF_B8G8R8A8.
@@ -637,9 +637,6 @@ private:
 	/** Specifies whether or not a point falls within any tagged drag regions that are draggable. */
 	bool IsInDragRegion(const FIntPoint& Point);
 
-	/** Used to let us correctly render the web texture for the accelerated render path */
-	TOptional<FSlateRenderTransform> GetWebBrowserRenderTransform() const;
-
 	bool BlockInputInDirectHwndMode() const;
 
 #if PLATFORM_WINDOWS
@@ -827,8 +824,9 @@ private:
 	TArray<FWebBrowserDragRegion> DragRegions;
 
 	TWeakPtr<SWindow> ParentWindow;
-	FCEFWebBrowserWindowRHIHelper* RHIRenderHelper;
-	
+
+	TUniquePtr<FCEFWebBrowserWindowRHIHelper> RHIRenderHelper;
+
 #if PLATFORM_WINDOWS || PLATFORM_MAC
 	bool bInDirectHwndMode;
 #endif
diff --git a/Engine/Source/Runtime/WebBrowser/Private/CEF/CEFWebBrowserWindowRHIHelper.cpp b/Engine/Source/Runtime/WebBrowser/Private/CEF/CEFWebBrowserWindowRHIHelper.cpp
--- a/Engine/Source/Runtime/WebBrowser/Private/CEF/CEFWebBrowserWindowRHIHelper.cpp
+++ b/Engine/Source/Runtime/WebBrowser/Private/CEF/CEFWebBrowserWindowRHIHelper.cpp
@@ -16,183 +16,302 @@
 
 #if PLATFORM_WINDOWS
 #include "Windows/AllowWindowsPlatformTypes.h"
-#include "d3d11.h"
 #include "d3d11_1.h"
+#include "d3d12.h"
+#include "dxgi1_4.h"
 #include "Windows/HideWindowsPlatformTypes.h"
-#endif		// PLATFORM_WINDOWS
 
-
-FCEFWebBrowserWindowRHIHelper::~FCEFWebBrowserWindowRHIHelper()
+static TComPtr<IDXGIAdapter> GetDXGIAdapterFromRHI()
 {
+	TComPtr<IDXGIAdapter> DXGIAdapter;
+
+	const ERHIInterfaceType RHIType = RHIGetInterfaceType();
+	if (RHIType == ERHIInterfaceType::D3D11)
+	{
+		TComPtr<IDXGIDevice> DXGIDevice;
+		if (HRESULT Hr = static_cast<ID3D11Device*>(GDynamicRHI->RHIGetNativeDevice())->QueryInterface(IID_PPV_ARGS(&DXGIDevice)); FAILED(Hr))
+		{
+			UE_LOG(LogWebBrowser, Error, TEXT("GetDXGIAdapter() - - ID3D11Device::QueryInterface 0x%x"), Hr);
+			return DXGIAdapter;
+		}
+		if (HRESULT Hr = DXGIDevice->GetAdapter(&DXGIAdapter); FAILED(Hr))
+		{
+			UE_LOG(LogWebBrowser, Error, TEXT("GetDXGIAdapter() - - IDXGIDevice::GetAdapter 0x%x"), Hr);
+			return DXGIAdapter;
+		}
+	}
+	else if (RHIType == ERHIInterfaceType::D3D12)
+	{
+		const LUID Luid = static_cast<ID3D12Device*>(GDynamicRHI->RHIGetNativeDevice())->GetAdapterLuid();
+
+		TComPtr<IDXGIFactory4> DXGIFactory;
+		if (HRESULT Hr = CreateDXGIFactory1(IID_PPV_ARGS(&DXGIFactory)); FAILED(Hr))
+		{
+			UE_LOG(LogWebBrowser, Error, TEXT("GetDXGIAdapter() - - CreateDXGIFactory1 0x%x"), Hr);
+			return DXGIAdapter;
+		}
+		if (HRESULT Hr = DXGIFactory->EnumAdapterByLuid(Luid, IID_PPV_ARGS(&DXGIAdapter)); FAILED(Hr))
+		{
+			UE_LOG(LogWebBrowser, Error, TEXT("GetDXGIAdapter() - - IDXGIFactory4::EnumAdapterByLuid 0x%x"), Hr);
+			return DXGIAdapter;
+		}
+	}
+
+	return DXGIAdapter;
 }
+#endif // PLATFORM_WINDOWS
 
-bool FCEFWebBrowserWindowRHIHelper::BUseRHIRenderer()
+
+FCEFWebBrowserWindowRHIHelper::FCEFWebBrowserWindowRHIHelper()
 {
 #if WITH_ENGINE
-	if (GDynamicRHI != nullptr && FCEFWebBrowserWindow::CanSupportAcceleratedPaint())
+#if PLATFORM_WINDOWS
+	// We need to get access to the underlying D3D device from RHI to create a new device with the same adapter
+	TComPtr<IDXGIAdapter> DXGIAdapter = GetDXGIAdapterFromRHI();
+
+	const D3D_FEATURE_LEVEL FeatureLevels[] =
 	{
-		return RHIGetInterfaceType() == ERHIInterfaceType::D3D11;
+		D3D_FEATURE_LEVEL_11_1,
+		D3D_FEATURE_LEVEL_11_0,
+	};
+
+	if (HRESULT Hr = D3D11CreateDevice(DXGIAdapter.Get(), DXGIAdapter.Get() ? D3D_DRIVER_TYPE_UNKNOWN : D3D_DRIVER_TYPE_HARDWARE, nullptr, 0, FeatureLevels, UE_ARRAY_COUNT(FeatureLevels), D3D11_SDK_VERSION,
+		&D3D11Device, nullptr, &D3D11DeviceContext); FAILED(Hr))
+	{
+		UE_LOG(LogWebBrowser, Error, TEXT("FCEFWebBrowserWindowRHIHelper::FCEFWebBrowserWindowRHIHelper() - - D3D11CreateDevice 0x%x"), Hr);
+		return;
 	}
-#endif
-	return false;
+#endif // PLATFORM_WINDOWS
+#endif // WITH_ENGINE
 }
 
-void FCEFWebBrowserWindowRHIHelper::UpdateCachedGeometry(const FGeometry& AllottedGeometryIn)
+bool FCEFWebBrowserWindowRHIHelper::BUseSupportedRHIRenderer()
 {
-	AllottedGeometry = AllottedGeometryIn;
+#if WITH_ENGINE
+	const ERHIInterfaceType RHIType = RHIGetInterfaceType();
+	// We only support D3D RHIs as we rely on the OpenSharedResource1 interface
+	return (RHIType == ERHIInterfaceType::D3D11 || RHIType == ERHIInterfaceType::D3D12);
+#else
+	return false;
+#endif // WITH_ENGINE
 }
 
-TOptional<FSlateRenderTransform> FCEFWebBrowserWindowRHIHelper::GetWebBrowserRenderTransform() const
+FSlateUpdatableTexture* FCEFWebBrowserWindowRHIHelper::CreateSlateUpdatableTexture(const FUintPoint& TextureSize)
 {
-	return FSlateRenderTransform(Concatenate(FScale2D(1, -1), FVector2D(0, AllottedGeometry.GetLocalSize().Y)));
+#if WITH_ENGINE
+	check(BUseSupportedRHIRenderer());
+#if PLATFORM_WINDOWS
+	// On Windows we need the slate texture to be shared so we can access it from the D3D device that lives in the OnAcceleratedPaint thread
+	FSlateTexture2DRHIRef* SlateTexture = new FSlateTexture2DRHIRef(TextureSize.X, TextureSize.Y, PF_B8G8R8A8, nullptr, ETextureCreateFlags::Shared | ETextureCreateFlags::RenderTargetable, true);
+	BeginInitResource(SlateTexture);
+	return SlateTexture;
+#else
+	return nullptr;
+#endif // PLATFORM_WINDOWS
+#else
+	return nullptr;
+#endif // WITH_ENGINE
 }
 
-FSlateUpdatableTexture *FCEFWebBrowserWindowRHIHelper::CreateTexture(void* SharedHandle)
+void FCEFWebBrowserWindowRHIHelper::ReleaseSlateUpdatableTexture(FSlateUpdatableTexture* SlateTexture)
 {
 #if WITH_ENGINE
-	check(BUseRHIRenderer());
+	check(BUseSupportedRHIRenderer());
 #if PLATFORM_WINDOWS
-	ID3D11Device1* dev1 = nullptr;
-	ID3D11Device* D3D11Device = static_cast<ID3D11Device*>(GDynamicRHI->RHIGetNativeDevice());
-	HRESULT Hr = D3D11Device->QueryInterface(__uuidof(ID3D11Device1), reinterpret_cast<void**>(&dev1));
-	if (FAILED(Hr))
+	if (!SlateTexture)
+	{
+		return;
+	}
+
+	SlateTextureHandles.Remove(SlateTexture);
+
+	BeginReleaseResource(static_cast<FSlateTexture2DRHIRef*>(SlateTexture));
+	FlushRenderingCommands();
+
+	delete SlateTexture;
+#endif // PLATFORM_WINDOWS
+#endif // WITH_ENGINE
+}
+
+bool FCEFWebBrowserWindowRHIHelper::EnsureShareable(FSlateUpdatableTexture* SlateTexture)
+{
+#if WITH_ENGINE
+	if (!SlateTexture)
+	{
+		UE_LOG(LogWebBrowser, Error, TEXT("FCEFWebBrowserWindowRHIHelper::EnsureShareable() - Invalid SlateTexture"));
+		return false;
+	}
+
+	if (TSharedPtr<void>* HandlePtr = SlateTextureHandles.Find(SlateTexture); HandlePtr)
 	{
-		UE_LOG(LogWebBrowser, Error, TEXT("FCEFWebBrowserWindowRHIHelper::CreateTexture() - - ID3D11Device::QueryInterface"));
-		return nullptr;
+		// Null Handle means that extraction/initialization failed
+		return HandlePtr->Get() != nullptr;
 	}
 
-	ID3D11Texture2D* tex = nullptr;
-	Hr = dev1->OpenSharedResource1(SharedHandle, __uuidof(ID3D11Texture2D), (void**)(&tex));
-	if (FAILED(Hr))
+#if PLATFORM_WINDOWS
+	FTextureRHIRef Slate2DRef = static_cast<FSlateTexture2DRHIRef*>(SlateTexture)->GetRHIRef();
+	if (!Slate2DRef.IsValid() || !Slate2DRef->GetNativeResource())
 	{
-		UE_LOG(LogWebBrowser, Error, TEXT("FCEFWebBrowserWindowRHIHelper::CreateTexture() - - ID3D11Device::OpenSharedResource"));
-		return nullptr;
+		// If this happens frequently then we need to make the resource creation sync in CreateSlateUpdatableTexture
+		UE_LOG(LogWebBrowser, Error, TEXT("FCEFWebBrowserWindowRHIHelper::EnsureShareable() - SlateTexture is not ready!"));
+		return false;
 	}
 
-	D3D11_TEXTURE2D_DESC TexDesc;
-	tex->GetDesc(&TexDesc);
-	FSlateTexture2DRHIRef* NewTexture = new FSlateTexture2DRHIRef(TexDesc.Width, TexDesc.Height, PF_R8G8B8A8, nullptr, TexCreate_None, true);
-	if (IsInRenderingThread())
+	HANDLE SharedHandle = nullptr;
+
+	// We add SharedHandle to the map even if we couldn't extract it, so we don't continuously try and report failures
+	const ERHIInterfaceType RHIType = RHIGetInterfaceType();
+	if (RHIType == ERHIInterfaceType::D3D11)
 	{
-		NewTexture->InitResource(FRHICommandListImmediate::Get());
+		TComPtr<IDXGIResource> DXGIResource;
+		if (HRESULT Hr = static_cast<ID3D11Texture2D*>(Slate2DRef->GetNativeResource())->QueryInterface(IID_PPV_ARGS(&DXGIResource)); FAILED(Hr))
+		{
+			UE_LOG(LogWebBrowser, Error, TEXT("FCEFWebBrowserWindowRHIHelper::EnsureShareable() - - ID3D11Texture2D::QueryInterface 0x%x"), Hr);
+		}
+		else if (Hr = DXGIResource->GetSharedHandle(&SharedHandle); FAILED(Hr))
+		{
+			UE_LOG(LogWebBrowser, Error, TEXT("FCEFWebBrowserWindowRHIHelper::EnsureShareable() - - IDXGIResource::GetSharedHandle 0x%x"), Hr);
+		}
+
+		// GetSharedHandle returns a legacy non-NT handle, which must NOT be closed when we're done with it
+		SlateTextureHandles.Add(SlateTexture, MakeShareable(SharedHandle, [](HANDLE _LegacySharedHandle) {}));
 	}
-	else
+	else if (RHIType == ERHIInterfaceType::D3D12)
 	{
-		BeginInitResource(NewTexture);
+		if (HRESULT Hr = static_cast<ID3D12Device*>(GDynamicRHI->RHIGetNativeDevice())->CreateSharedHandle(static_cast<ID3D12Resource*>(Slate2DRef->GetNativeResource()), NULL, GENERIC_ALL, NULL, &SharedHandle); FAILED(Hr))
+		{
+			UE_LOG(LogWebBrowser, Error, TEXT("FCEFWebBrowserWindowRHIHelper::EnsureShareable() - - ID3D12Device::CreateSharedHandle 0x%x"), Hr);
+		}
+
+		// CreateSharedHandle returns a NT handle,. which must be closed when we're done with it
+		SlateTextureHandles.Add(SlateTexture, MakeShareable(SharedHandle, [](HANDLE _NTSharedHandle)
+			{
+				if (_NTSharedHandle)
+				{
+					CloseHandle(_NTSharedHandle);
+				}
+			}));
 	}
 
-	return NewTexture;
+	return SharedHandle != nullptr;
 #else
-	return nullptr;
-#endif
+	return false;
+#endif // PLATFORM_WINDOWS
 #else
-	return nullptr;
-#endif
+	return false;
+#endif // WITH_ENGINE
+
 }
 
-void FCEFWebBrowserWindowRHIHelper::UpdateSharedHandleTexture(void *SharedHandle, FSlateUpdatableTexture* SlateTexture, const FIntRect& DirtyIn)
+bool FCEFWebBrowserWindowRHIHelper::CopySharedTextureSync(FSlateUpdatableTexture* SlateTexture, void* SharedHandle, const FIntRect& DirtyIn)
 {
 #if WITH_ENGINE
-	check(BUseRHIRenderer());
+	check(BUseSupportedRHIRenderer());
 	FIntRect Dirty = DirtyIn;
 
+	if (!EnsureShareable(SlateTexture))
+	{
+		return false;
+	}
+
 #if PLATFORM_WINDOWS
-	ID3D11Device1* dev1 = nullptr;
-	ID3D11Device* D3D11Device = static_cast<ID3D11Device*>(GDynamicRHI->RHIGetNativeDevice());
-	HRESULT Hr = D3D11Device->QueryInterface(__uuidof(ID3D11Device1), reinterpret_cast<void**>(&dev1));
-	if (FAILED(Hr))
+	TComPtr<ID3D11Device1> D3D11Device1;
+	if (HRESULT Hr = D3D11Device->QueryInterface(IID_PPV_ARGS(&D3D11Device1)); FAILED(Hr))
 	{
-		UE_LOG(LogWebBrowser, Error, TEXT("FCEFWebBrowserWindowRHIHelper::UpdateSharedHandleTexture() - ID3D11Device::QueryInterface"), Hr);
-		return;
+		UE_LOG(LogWebBrowser, Error, TEXT("FCEFWebBrowserWindowRHIHelper::CopySharedTextureSync() - - ID3D11Device::QueryInterface 0x%x"), Hr);
+		return false;
 	}
 
-	ID3D11Texture2D* Tex = nullptr;
-	Hr = dev1->OpenSharedResource1(SharedHandle, __uuidof(ID3D11Texture2D), (void**)(&Tex));
-	if (FAILED(Hr))
+	TComPtr<ID3D11Texture2D> SourceTexture;
+	if (HRESULT Hr = D3D11Device1->OpenSharedResource1((HANDLE)SharedHandle, IID_PPV_ARGS(&SourceTexture)); FAILED(Hr))
 	{
-		UE_LOG(LogWebBrowser, Error, TEXT("FCEFWebBrowserWindowRHIHelper::UpdateSharedHandleTexture() - ID3D11Device::OpenSharedResource"), Hr);
-		return;
+		UE_LOG(LogWebBrowser, Error, TEXT("FCEFWebBrowserWindowRHIHelper::CopySharedTextureSync() - - ID3D11Device1::OpenSharedResource1 0x%x"), Hr);
+		return false;
 	}
 
-	D3D11_TEXTURE2D_DESC TexDesc;
-	Tex->GetDesc(&TexDesc);
-	if (SlateTexture->GetSlateResource()->GetWidth() != TexDesc.Width
-		|| SlateTexture->GetSlateResource()->GetHeight() != TexDesc.Height)
+	D3D11_TEXTURE2D_DESC SourceTextureDesc = {};
+	SourceTexture->GetDesc(&SourceTextureDesc);
+	if (SlateTexture->GetSlateResource()->GetWidth() != SourceTextureDesc.Width || SlateTexture->GetSlateResource()->GetHeight() != SourceTextureDesc.Height)
 	{
-		SlateTexture->ResizeTexture(TexDesc.Width, TexDesc.Height);
+		SlateTexture->ResizeTexture(SourceTextureDesc.Width, SourceTextureDesc.Height);
+		FlushRenderingCommands();
+
+		// The native resource has changed, we need to update the shared handle
+		SlateTextureHandles.Remove(SlateTexture);
+		EnsureShareable(SlateTexture);
+
 		Dirty = FIntRect();
 	}
 
-	TRefCountPtr<IDXGIKeyedMutex> KeyedMutex;
-	Hr = Tex->QueryInterface(_uuidof(IDXGIKeyedMutex), (void**)&KeyedMutex);
-	if (FAILED(Hr))
+	TComPtr<ID3D11Texture2D> DestTexture;
+	const ERHIInterfaceType RHIType = RHIGetInterfaceType();
+	if (RHIType == ERHIInterfaceType::D3D11)
 	{
-		UE_LOG(LogWebBrowser, Error, TEXT("FCEFWebBrowserWindowRHIHelper::UpdateSharedHandleTexture() - ID3D11Texture2D::IDXGIKeyedMutex"), Hr);
-		return;
+		if (HRESULT Hr = D3D11Device->OpenSharedResource((HANDLE)SlateTextureHandles.FindRef(SlateTexture).Get(), IID_PPV_ARGS(&DestTexture)); FAILED(Hr))
+		{
+			UE_LOG(LogWebBrowser, Error, TEXT("FCEFWebBrowserWindowRHIHelper::CopySharedTextureSync() - - ID3D11Device::OpenSharedResource 0x%x"), Hr);
+			return false;
+		}
 	}
+	else if (RHIType == ERHIInterfaceType::D3D12)
+	{
+		if (HRESULT Hr = D3D11Device1->OpenSharedResource1((HANDLE)SlateTextureHandles.FindRef(SlateTexture).Get(), IID_PPV_ARGS(&DestTexture)); FAILED(Hr))
+		{
+			UE_LOG(LogWebBrowser, Error, TEXT("FCEFWebBrowserWindowRHIHelper::CopySharedTextureSync() - - ID3D11Device1::OpenSharedResource1 0x%x"), Hr);
+			return false;
+		}
+	}
+
+	// Copy the textures on the GPU
+	// Note that as of CEF v128 the Dirty rect sent by CEF always matches the whole SourceTexture no matter what part has actually changed,
+	// so we can just do full copies. If that were to change in the future we can start using CopySubresourceRegion instead
+	D3D11DeviceContext->CopyResource(DestTexture.Get(), SourceTexture.Get());
 
-	ENQUEUE_RENDER_COMMAND(CEFAcceleratedPaint)(
-		[KeyedMutex, SlateTexture, Dirty, Tex](FRHICommandList& RHICmdList)
+	// The CopyResource call is asynchronous by default, so we need to force a sync flush to ensure it completes before we exit this method
+	D3D11DeviceContext->Flush();
+
+	// Now wait on the GPU to complete the copy and flush
+	D3D11_QUERY_DESC QueryDesc = {};
+	QueryDesc.Query = D3D11_QUERY_EVENT;
+	TComPtr<ID3D11Query> Query;
+	if (HRESULT Hr = D3D11Device1->CreateQuery(&QueryDesc, &Query); FAILED(Hr))
+	{
+		UE_LOG(LogWebBrowser, Error, TEXT("FCEFWebBrowserWindowRHIHelper::CopySharedTextureSync() - ID3D11Device1::CreateQuery 0x%x"), Hr);
+		return false;
+	}
+	D3D11DeviceContext->End(Query.Get());
+	BOOL bIsDone = false;
+	HRESULT Hr = S_OK;
+	for (;;)
+	{
+		Hr = D3D11DeviceContext->GetData(Query.Get(), &bIsDone, sizeof(bIsDone), 0);
+		if (FAILED(Hr))
 		{
-			if (KeyedMutex)
-			{
-				if (KeyedMutex->AcquireSync(1, 16) == S_OK)
-				{
-					D3D11_BOX Region;
-					Region.front = 0;
-					Region.back = 1;
-					if (Dirty.Area() > 0)
-					{
-						Region.left = Dirty.Min.X;
-						Region.top = Dirty.Min.Y;
-						Region.right = Dirty.Max.X;
-						Region.bottom = Dirty.Max.Y;
-					}
-					else
-					{
-						Region.left = 0;
-						Region.right = SlateTexture->GetSlateResource()->GetWidth();
-						Region.top = 0;
-						Region.bottom = SlateTexture->GetSlateResource()->GetHeight();
-					}
-
-					ID3D11DeviceContext* DX11DeviceContext = nullptr;
-					ID3D11Device* D3D11Device = static_cast<ID3D11Device*>(GDynamicRHI->RHIGetNativeDevice());
-					D3D11Device->GetImmediateContext(&DX11DeviceContext);
-					check(DX11DeviceContext);
-
-					FSlateTexture2DRHIRef* SlateRHITexture = static_cast<FSlateTexture2DRHIRef*>(SlateTexture);
-					check(SlateRHITexture);
-					auto Slate2DRef = SlateRHITexture->GetRHIRef();
-					if (Slate2DRef && Slate2DRef.IsValid())
-					{
-						auto Slate2DTexture = Slate2DRef->GetTexture2D();
-						if (Slate2DTexture)
-						{
-							ID3D11Texture2D* D3DTexture = (ID3D11Texture2D*)Slate2DTexture->GetNativeResource();
-
-							DX11DeviceContext->CopySubresourceRegion(D3DTexture, 0, Region.left, Region.top, Region.front, Tex, 0, &Region);
-						}
-					}
-
-					KeyedMutex->ReleaseSync(0);
-				}
-				else
-				{
-					UE_LOG(LogWebBrowser, Verbose, TEXT("FCEFWebBrowserWindowRHIHelper::UpdateSharedHandleTexture() - failed getting sync"), E_FAIL);
-				}
-			}
-			else
-			{
-				UE_LOG(LogWebBrowser, Verbose, TEXT("FCEFWebBrowserWindowRHIHelper::UpdateSharedHandleTexture() - missing KeyedMutex"), E_FAIL);
-				return;
-			}
-	});
+			UE_LOG(LogWebBrowser, Error, TEXT("FCEFWebBrowserWindowRHIHelper::CopySharedTextureSync() - ID3D11DeviceContext::GetData  0x%x"), Hr);
+			break;
+		}
+
+		// We need to check for S_OK specifically as S_FALSE is also considered a success return code
+		if (Hr == S_OK && bIsDone)
+		{
+			break;
+		}
+
+		// Idle wait before retrying
+		if (!SwitchToThread())
+		{
+			Sleep(1); // 1ms
+		}
+	}
+
+	return true;
 #else
-	UE_LOG(LogWebBrowser, Error, TEXT("FCEFWebBrowserWindowRHIHelper::UpdateSharedHandleTexture() - missing implementation"));
+	UE_LOG(LogWebBrowser, Error, TEXT("FCEFWebBrowserWindowRHIHelper::CopySharedTextureSync() - missing implementation"));
+	return false;
 #endif // PLATFORM_WINDOWS
 #else
-	UE_LOG(LogWebBrowser, Error, TEXT("FCEFWebBrowserWindowRHIHelper::UpdateSharedHandleTexture() - unsupported usage, RHI renderer but missing engine"));
+	UE_LOG(LogWebBrowser, Error, TEXT("FCEFWebBrowserWindowRHIHelper::CopySharedTextureSync() - unsupported usage, RHI renderer but missing engine"));
+	return false;
 #endif // WITH_ENGINE
 }
 
diff --git a/Engine/Source/Runtime/WebBrowser/Private/CEF/CEFWebBrowserWindowRHIHelper.h b/Engine/Source/Runtime/WebBrowser/Private/CEF/CEFWebBrowserWindowRHIHelper.h
--- a/Engine/Source/Runtime/WebBrowser/Private/CEF/CEFWebBrowserWindowRHIHelper.h
+++ b/Engine/Source/Runtime/WebBrowser/Private/CEF/CEFWebBrowserWindowRHIHelper.h
@@ -8,29 +8,43 @@
 
 #include "Layout/Geometry.h"
 
-class FSlateTexture2DRHIRef;
 class FSlateUpdatableTexture;
 
+#if PLATFORM_WINDOWS
+#include "Windows/WindowsHWrapper.h"
+#include "Microsoft/COMPointer.h"
+
+#include "Windows/AllowWindowsPlatformTypes.h"
+#include "d3d11.h"
+#include "Windows/HideWindowsPlatformTypes.h"
+#endif // PLATFORM_WINDOWS
+
 /**
  * Implementation of RHI renderer details for the CEF accelerated rendering path
  */
 class FCEFWebBrowserWindowRHIHelper
 {
-
 public:
-	/** Virtual Destructor. */
-	virtual ~FCEFWebBrowserWindowRHIHelper();
+	FCEFWebBrowserWindowRHIHelper();
+	virtual ~FCEFWebBrowserWindowRHIHelper() = default;
 
+	static bool BUseSupportedRHIRenderer();
 
 public:
-	static bool BUseRHIRenderer();
-	FSlateUpdatableTexture* CreateTexture(void *ShareHandle);
-	void UpdateSharedHandleTexture(void* SharedHandle, FSlateUpdatableTexture* SlateTexture, const FIntRect& DirtyIn);
-	void UpdateCachedGeometry(const FGeometry& AllottedGeometry);
-	TOptional<FSlateRenderTransform> GetWebBrowserRenderTransform() const;
+	FSlateUpdatableTexture* CreateSlateUpdatableTexture(const FUintPoint& TextureSize);
+	void ReleaseSlateUpdatableTexture(FSlateUpdatableTexture* SlateTexture);
+
+	bool CopySharedTextureSync(FSlateUpdatableTexture* SlateTexture, void* SharedHandle, const FIntRect& DirtyIn);
 
 private:
-	FGeometry AllottedGeometry;
+	bool EnsureShareable(FSlateUpdatableTexture* SlateTexture);
+
+	TMap<FSlateUpdatableTexture*, TSharedPtr<void>> SlateTextureHandles;
+
+#if PLATFORM_WINDOWS
+	TComPtr<ID3D11Device> D3D11Device;
+	TComPtr<ID3D11DeviceContext> D3D11DeviceContext;
+#endif // PLATFORM_WINDOWS
 };
 
 #endif
diff --git a/Engine/Source/Runtime/WebBrowser/Private/SWebBrowser.cpp b/Engine/Source/Runtime/WebBrowser/Private/SWebBrowser.cpp
--- a/Engine/Source/Runtime/WebBrowser/Private/SWebBrowser.cpp
+++ b/Engine/Source/Runtime/WebBrowser/Private/SWebBrowser.cpp
@@ -10,7 +10,7 @@
 #include "Framework/Docking/TabManager.h"
 #include "Async/Async.h"
 #include "CEF/CEFWebBrowserWindow.h"
-#endif
+#endif // WITH_CEF3
 #include "WebBrowserModule.h"
 #include "WebBrowserSingleton.h"
 
@@ -25,7 +25,6 @@ SWebBrowser::~SWebBrowser()
 }
 
 #if WITH_CEF3
-
 static TSharedRef<FCEFWebBrowserWindow> GetCefWebBrowserWindowChecked(const TSharedPtr<IWebBrowserWindow>& WebBrowserWindow)
 {
 	check(WebBrowserWindow.IsValid());
@@ -35,8 +34,7 @@ static TSharedRef<FCEFWebBrowserWindow> GetCefWebBrowserWindowChecked(const TSha
 	
 	return CefWebBrowserWindow.ToSharedRef();
 }
-
-#endif
+#endif // WITH_CEF3
 
 static FWebBrowserSingleton& GetWebBrowserSingletonChecked()
 {
@@ -50,16 +48,12 @@ static FWebBrowserSingleton& GetWebBrowserSingletonChecked()
 {
 	SCompoundWidget::Tick(AllottedGeometry, InCurrentTime, InDeltaTime);
 
-	
 #if WITH_EDITOR
 #if WITH_CEF3
-
-	// If we're using captured CEF buffer / smooth resizing.
-
 	if (WebBrowserWindow.IsValid()) 
 	{
 		if (const TSharedRef<FCEFWebBrowserWindow> CefWebBrowserWindow = GetCefWebBrowserWindowChecked(WebBrowserWindow);
-			CefWebBrowserWindow->IsUsingCapturedCefBuffer())
+			CefWebBrowserWindow->HasResizeBug())
 		{
 			// NOTE - We won't end up here if the CEF version doesn't require it.
 			
@@ -82,10 +76,8 @@ static FWebBrowserSingleton& GetWebBrowserSingletonChecked()
 				// Always seems necessary (even if viewport size was just changed above.)
 				CefWebBrowserWindow->SetNeedsResize(true);
 			}
-
 	
-#if PLATFORM_WINDOWS // TODO - Test MacOS and Linux.
-
+#if PLATFORM_WINDOWS
 			// * This is necessary during active widget resizes, only if the widget is in a floating window, and only when on Windows. *
 			// When we are (1) running on Windows (2) not docked in the Editor (floating), and (3) our window is actively being resized - then the FTS
 			// Ticker interface on the Web Browser Singleton will be suspended. But our Tick() function here will not be suspended. So under these and
@@ -100,26 +92,21 @@ static FWebBrowserSingleton& GetWebBrowserSingletonChecked()
 			{
 				WebBrowserSingleton.Tick(InDeltaTime);
 			}
-	
-#endif
+#endif // PLATFORM_WINDOWS
 		}
 	}
-	
-#endif
-#endif
+#endif // WITH_CEF3
+#endif // WITH_EDITOR
 }
 
 /*virtual*/ int32 SWebBrowser::OnPaint(const FPaintArgs& Args, const FGeometry& AllottedGeometry, const FSlateRect& MyCullingRect, FSlateWindowElementList& OutDrawElements, int32 LayerId, const FWidgetStyle& InWidgetStyle, bool bParentEnabled) const
 {
 #if WITH_EDITOR
 #if WITH_CEF3
-	
-	// If we're using captured CEF buffer / smooth resizing..
-
 	if (WebBrowserWindow.IsValid())
 	{
 		if (const TSharedPtr<FCEFWebBrowserWindow> CefWebBrowserWindow = GetCefWebBrowserWindowChecked(WebBrowserWindow);
-			CefWebBrowserWindow->IsUsingCapturedCefBuffer())
+			CefWebBrowserWindow->HasResizeBug())
 		{
 			// NOTE - We won't end up here if the CEF version doesn't require it.
 			
@@ -160,10 +147,8 @@ static FWebBrowserSingleton& GetWebBrowserSingletonChecked()
 			}
 		}
 	}
-	
-#endif
-#endif
-
+#endif // WITH_CEF3
+#endif // WITH_EDITOR
 
 	return SCompoundWidget::OnPaint(Args, AllottedGeometry, MyCullingRect, OutDrawElements, LayerId, InWidgetStyle, bParentEnabled);
 }
@@ -306,12 +291,12 @@ void SWebBrowser::Construct(const FArguments& InArgs, const TSharedPtr<IWebBrows
 			(GIsPlayInEditorWorld ? false : InArgs._UseSmoothResizing)
 #else
 			false
-#endif
+#endif // WITH_EDITOR
 			;
 		
 		GetCefWebBrowserWindowChecked(WebBrowserWindow)->SetUseCapturedCefBuffer(bUseSmoothScrolling);
 	}
-#endif
+#endif // WITH_CEF3
 }
 
 void SWebBrowser::LoadURL(FString NewURL)
diff --git a/Engine/Source/Runtime/WebBrowser/WebBrowser.Build.cs b/Engine/Source/Runtime/WebBrowser/WebBrowser.Build.cs
--- a/Engine/Source/Runtime/WebBrowser/WebBrowser.Build.cs
+++ b/Engine/Source/Runtime/WebBrowser/WebBrowser.Build.cs
@@ -86,14 +86,17 @@ public class WebBrowser : ModuleRules
 			PublicSystemLibraries.Add("libjnigraphics");
 		}
 
+		if (Target.Platform == UnrealTargetPlatform.Win64)
+		{
+			AddEngineThirdPartyPrivateStaticDependencies(Target, "DX11");
+		}
+
 		if (Target.Platform == UnrealTargetPlatform.Win64
 		||  Target.Platform == UnrealTargetPlatform.Mac
 		||  Target.Platform == UnrealTargetPlatform.Linux)
 		{
 			PrivateDependencyModuleNames.Add("CEF3Utils");
-			AddEngineThirdPartyPrivateStaticDependencies(Target,
-				"CEF3"
-				);
+			AddEngineThirdPartyPrivateStaticDependencies(Target, "CEF3");
 
 			if (Target.Type != TargetType.Server)
 			{
