WebBrowser: disable CEF accelerated rendering for automation tests

The CEF/Chromium GPU process often fails to initialize, likely due to not having access to some graphics resources.

#rnx

diff --git a/Engine/Source/Runtime/WebBrowser/Private/CEF/CEFWebBrowserWindowRHIHelper.cpp b/Engine/Source/Runtime/WebBrowser/Private/CEF/CEFWebBrowserWindowRHIHelper.cpp
--- a/Engine/Source/Runtime/WebBrowser/Private/CEF/CEFWebBrowserWindowRHIHelper.cpp
+++ b/Engine/Source/Runtime/WebBrowser/Private/CEF/CEFWebBrowserWindowRHIHelper.cpp
@@ -90,7 +90,9 @@ bool FCEFWebBrowserWindowRHIHelper::BUseSupportedRHIRenderer()
 #if WITH_ENGINE
 	const ERHIInterfaceType RHIType = RHIGetInterfaceType();
 	// We only support D3D RHIs as we rely on the OpenSharedResource1 interface
-	return (RHIType == ERHIInterfaceType::D3D11 || RHIType == ERHIInterfaceType::D3D12);
+	return (RHIType == ERHIInterfaceType::D3D11 || RHIType == ERHIInterfaceType::D3D12)
+		// Disable during automation as the GPU process often fails to initialize on build machines
+		&& !GIsAutomationTesting && !GIsBuildMachine;
 #else
 	return false;
 #endif // WITH_ENGINE
