[Mac WebBrowser] Fix an issue where UObject references could be destroyed prior to web browser shut down, resulting in dangling pointers.

#rnx

diff --git a/Engine/Source/Runtime/WebBrowser/Private/Apple/ApplePlatformWebBrowser.cpp b/Engine/Source/Runtime/WebBrowser/Private/Apple/ApplePlatformWebBrowser.cpp
--- a/Engine/Source/Runtime/WebBrowser/Private/Apple/ApplePlatformWebBrowser.cpp
+++ b/Engine/Source/Runtime/WebBrowser/Private/Apple/ApplePlatformWebBrowser.cpp
@@ -103,7 +103,7 @@ class SAppleWebBrowserWidget : public SViewport
 		// create external texture
 		WebBrowserTexture = NewObject<UWebBrowserTexture>((UObject*)GetTransientPackage(), NAME_None, RF_Transient | RF_Public);
 
-		if (WebBrowserTexture != nullptr)
+		if (WebBrowserTexture.IsValid())
 		{
 			WebBrowserTexture->UpdateResource();
 			WebBrowserTexture->AddToRoot();
@@ -118,15 +118,15 @@ class SAppleWebBrowserWidget : public SViewport
 			// create wrapper material
 			WebBrowserMaterial = UMaterialInstanceDynamic::Create(DefaultWBMaterial, nullptr);
 
-			if (WebBrowserMaterial)
+			if (WebBrowserMaterial.IsValid())
 			{
-				WebBrowserMaterial->SetTextureParameterValue("SlateUI", WebBrowserTexture);
+				WebBrowserMaterial->SetTextureParameterValue("SlateUI", WebBrowserTexture.Get());
 				WebBrowserMaterial->AddToRoot();
 
 				// create Slate brush
 				WebBrowserBrush = MakeShareable(new FSlateBrush());
 				{
-					WebBrowserBrush->SetResourceObject(WebBrowserMaterial);
+					WebBrowserBrush->SetResourceObject(WebBrowserMaterial.Get());
 				}
 			}
 		}
@@ -286,7 +286,7 @@ class SAppleWebBrowserWidget : public SViewport
 #if !PLATFORM_TVOS
 			if (Is3DBrowser)
 			{
-				if (WebBrowserTexture)
+				if (WebBrowserTexture.IsValid())
 				{					
 					TSharedPtr<FWebBrowserTextureSample, ESPMode::ThreadSafe> WebBrowserTextureSample;
 					WebBrowserTextureSamplesQueue->Peek(WebBrowserTextureSample);
@@ -294,10 +294,8 @@ class SAppleWebBrowserWidget : public SViewport
 					WebBrowserTexture->TickResource(WebBrowserTextureSample);
 				}
 
-				if (WebBrowserTexture != nullptr)
+				if (WebBrowserTexture.IsValid())
 				{
-
-					
 					TSharedPtr<FWebBrowserWindow> WebBrowserWindow = WebBrowserWindowPtr.Pin();
 					bool TextureWritten = false;
 										
@@ -408,12 +406,12 @@ class SAppleWebBrowserWidget : public SViewport
 		{
 			WebBrowserBrush->SetResourceObject(nullptr);	
 		}
-		if (WebBrowserMaterial != nullptr)
+		if (WebBrowserMaterial.IsValid())
 		{
 			WebBrowserMaterial->RemoveFromRoot();
 			WebBrowserMaterial = nullptr;
 		}
-		if (WebBrowserTexture != nullptr)
+		if (WebBrowserTexture.IsValid())
 		{
 			WebBrowserTexture->RemoveFromRoot();
 			WebBrowserTexture = nullptr;
@@ -686,10 +684,10 @@ private:
 
 #if !PLATFORM_TVOS
 	/** The external texture to render the webbrowser output. */
-	UWebBrowserTexture* WebBrowserTexture;
+	TWeakObjectPtr<UWebBrowserTexture> WebBrowserTexture;
 
 	/** The material for the external texture. */
-	UMaterialInstanceDynamic* WebBrowserMaterial;
+	TWeakObjectPtr<UMaterialInstanceDynamic> WebBrowserMaterial;
 
 	/** The Slate brush that renders the material. */
 	TSharedPtr<FSlateBrush> WebBrowserBrush;
