WebBrowser: Make sure UE and CEF run on the same adapter in multi-GPU setups

#tests forced a LUID on the UE side, verified that it was correctly transferred to CEF and used by the GPU process (using the about://gpu page)
#rnx

diff --git a/Engine/Source/Runtime/WebBrowser/Private/CEF/CEFBrowserApp.cpp b/Engine/Source/Runtime/WebBrowser/Private/CEF/CEFBrowserApp.cpp
--- a/Engine/Source/Runtime/WebBrowser/Private/CEF/CEFBrowserApp.cpp
+++ b/Engine/Source/Runtime/WebBrowser/Private/CEF/CEFBrowserApp.cpp
@@ -5,6 +5,7 @@
 
 #if WITH_CEF3
 #include "WebBrowserLog.h"
+#include "CEF/CEFWebBrowserWindowRHIHelper.h"
 
 //#define DEBUG_CEFMESSAGELOOP_FRAMERATE 1 // uncomment this to have debug spew about the FPS we call the CefDoMessageLoopWork function
 
@@ -33,6 +34,16 @@ void FCEFBrowserApp::OnBeforeCommandLineProcessing(const CefString& ProcessType,
 		UE_LOG(LogCEFBrowser, Log, TEXT("CEF GPU acceleration enabled"));
 		CommandLine->AppendSwitch("enable-gpu");
 		CommandLine->AppendSwitch("enable-gpu-compositing");
+
+#if PLATFORM_WINDOWS
+		// Make sure CEF uses the same adapter as UE for shared textures to work,
+		// which requires using ANGLE's D3D11 backend
+		CommandLine->AppendSwitchWithValue("use-angle", "d3d11");
+		const uint64_t AdapterLuid = FCEFWebBrowserWindowRHIHelper::GetRHIAdapterLuid();
+		CommandLine->AppendSwitchWithValue("use-adapter-luid",
+			*FString::Printf(TEXT("%u,%u"), static_cast<uint32_t>(AdapterLuid >> 32), static_cast<uint32_t>(AdapterLuid & 0xFFFFFFFF)));
+#endif
+
 	}
 	else
 	{
diff --git a/Engine/Source/Runtime/WebBrowser/Private/CEF/CEFWebBrowserWindowRHIHelper.cpp b/Engine/Source/Runtime/WebBrowser/Private/CEF/CEFWebBrowserWindowRHIHelper.cpp
--- a/Engine/Source/Runtime/WebBrowser/Private/CEF/CEFWebBrowserWindowRHIHelper.cpp
+++ b/Engine/Source/Runtime/WebBrowser/Private/CEF/CEFWebBrowserWindowRHIHelper.cpp
@@ -98,6 +98,36 @@ bool FCEFWebBrowserWindowRHIHelper::BUseSupportedRHIRenderer()
 #endif // WITH_ENGINE
 }
 
+uint64_t FCEFWebBrowserWindowRHIHelper::GetRHIAdapterLuid()
+{
+	uint64_t AdapterLuid = 0;
+
+#if WITH_ENGINE
+#if PLATFORM_WINDOWS
+	LUID Luid = {};
+
+	const ERHIInterfaceType RHIType = RHIGetInterfaceType();
+	if (RHIType == ERHIInterfaceType::D3D11)
+	{
+		TComPtr<IDXGIAdapter> DXGIAdapter = GetDXGIAdapterFromRHI();
+		DXGI_ADAPTER_DESC Desc = {};
+		if (SUCCEEDED(DXGIAdapter->GetDesc(&Desc)))
+		{
+			Luid = Desc.AdapterLuid;
+		}
+	}
+	else if (RHIType == ERHIInterfaceType::D3D12)
+	{
+		Luid = static_cast<ID3D12Device*>(GDynamicRHI->RHIGetNativeDevice())->GetAdapterLuid();
+	}
+
+	AdapterLuid = (static_cast<uint64_t>(Luid.HighPart) << 32) | static_cast<uint64_t>(Luid.LowPart);
+#endif // PLATFORM_WINDOWS
+#endif // WITH_ENGINE
+
+	return AdapterLuid;
+}
+
 FSlateUpdatableTexture* FCEFWebBrowserWindowRHIHelper::CreateSlateUpdatableTexture(const FUintPoint& TextureSize)
 {
 #if WITH_ENGINE
diff --git a/Engine/Source/Runtime/WebBrowser/Private/CEF/CEFWebBrowserWindowRHIHelper.h b/Engine/Source/Runtime/WebBrowser/Private/CEF/CEFWebBrowserWindowRHIHelper.h
--- a/Engine/Source/Runtime/WebBrowser/Private/CEF/CEFWebBrowserWindowRHIHelper.h
+++ b/Engine/Source/Runtime/WebBrowser/Private/CEF/CEFWebBrowserWindowRHIHelper.h
@@ -29,6 +29,7 @@ public:
 	virtual ~FCEFWebBrowserWindowRHIHelper() = default;
 
 	static bool BUseSupportedRHIRenderer();
+	static uint64_t GetRHIAdapterLuid();
 
 public:
 	FSlateUpdatableTexture* CreateSlateUpdatableTexture(const FUintPoint& TextureSize);
