Ignore abnormal requests in EDA causing macOS warnings popups

#rnx

diff --git a/Engine/Source/Runtime/WebBrowser/Private/Apple/ApplePlatformWebBrowser.cpp b/Engine/Source/Runtime/WebBrowser/Private/Apple/ApplePlatformWebBrowser.cpp
--- a/Engine/Source/Runtime/WebBrowser/Private/Apple/ApplePlatformWebBrowser.cpp
+++ b/Engine/Source/Runtime/WebBrowser/Private/Apple/ApplePlatformWebBrowser.cpp
@@ -484,7 +484,8 @@ class SAppleWebBrowserWidget : public SViewport
 
 	bool HandleShouldOverrideUrlLoading(const FString& Url)
 	{
-		if (WebBrowserWindowPtr.IsValid())
+		// UE-347936 Abnormal requests in EDA webpage, ignore these
+		if (WebBrowserWindowPtr.IsValid() && !Url.Equals(TEXT("about:blank")) && !Url.Equals(TEXT("about:srcdoc")))
 		{
 			// Capture vars needed for AsyncTask
 			NSString* UrlString = [NSString stringWithUTF8String : TCHAR_TO_UTF8(*Url)];
