Improve the foreground tab tracking workaround for the visibility of the web browser window.
Now it also work for when the tab live under another major tab.

diff --git a/Engine/Source/Runtime/WebBrowser/Private/Apple/ApplePlatformWebBrowser.cpp b/Engine/Source/Runtime/WebBrowser/Private/Apple/ApplePlatformWebBrowser.cpp
--- a/Engine/Source/Runtime/WebBrowser/Private/Apple/ApplePlatformWebBrowser.cpp
+++ b/Engine/Source/Runtime/WebBrowser/Private/Apple/ApplePlatformWebBrowser.cpp
@@ -23,6 +23,7 @@
 #include "Framework/Docking/TabManager.h"
 #include "Rendering/SlateRenderer.h"
 #include "Textures/SlateUpdatableTexture.h"
+#include "Widgets/Docking/SDockTab.h"
 
 #import <MetalKit/MetalKit.h>
 #include "ExternalTexture.h"
@@ -656,6 +657,39 @@ class SAppleWebBrowserWidget : public SViewport
 			// this should only be called when the active browser has an active parent dock tab
 			check(CurrentTab.IsValid());
 			
+			// Deal with a new major tab being foregrounded that could host our CurrentTab
+			if (ActiveTab && ActiveTab->GetTabRole() == ETabRole::MajorTab)
+			{
+				if (TSharedPtr<FTabManager> CurrentTabManager = CurrentTab->GetTabManagerPtr())
+				{
+					if (TSharedPtr<SDockTab> CurrentMajorTab = FGlobalTabmanager::Get()->GetMajorTabForTabManager(CurrentTabManager.ToSharedRef()))
+					{
+						if (ActiveTab == CurrentMajorTab)
+						{
+							if (CurrentTab->IsForeground())
+							{
+								[WebViewWrapper setVisibility : YES] ;
+							}
+						}
+					}
+				}
+			}
+
+			// Deal with a major tab being move to the background that could host our CurrentTab
+			if (PreviouslyActiveTab && PreviouslyActiveTab->GetTabRole() == ETabRole::MajorTab)
+			{
+				if (TSharedPtr<FTabManager> CurrentTabManager = CurrentTab->GetTabManagerPtr())
+				{
+					if (TSharedPtr<SDockTab> CurrentMajorTab = FGlobalTabmanager::Get()->GetMajorTabForTabManager(CurrentTabManager.ToSharedRef()))
+					{
+						if (PreviouslyActiveTab == CurrentMajorTab)
+						{
+							[WebViewWrapper setVisibility : NO ];
+						}
+					}
+				}
+			}
+
 			if(ActiveTab == CurrentTab)
 			{
 				[WebViewWrapper setVisibility: YES ];
