[Apple WebBrowser] Replace hard reference to the web browser's parent window with a weak ptr.

#rnx

diff --git a/Engine/Source/Runtime/WebBrowser/Private/Apple/ApplePlatformWebBrowser.cpp b/Engine/Source/Runtime/WebBrowser/Private/Apple/ApplePlatformWebBrowser.cpp
--- a/Engine/Source/Runtime/WebBrowser/Private/Apple/ApplePlatformWebBrowser.cpp
+++ b/Engine/Source/Runtime/WebBrowser/Private/Apple/ApplePlatformWebBrowser.cpp
@@ -138,7 +138,7 @@ class SAppleWebBrowserWidget : public SViewport
 	void UpdateFrameWithViewport(FIntPoint WindowSize, FIntPoint WindowPos) {
 		CGRect NewFrame;
 #if PLATFORM_MAC
-		TSharedPtr<FGenericWindow> EnclosedWindow = WebBrowserWindowPtr.Pin()->GetParentWindow().Get()->GetNativeWindow();
+		TSharedPtr<FGenericWindow> EnclosedWindow = WebBrowserWindowPtr.Pin()->GetParentWindow()->GetNativeWindow();
 		
 		if(!Is3DBrowser) 
 		{
@@ -256,10 +256,11 @@ class SAppleWebBrowserWidget : public SViewport
 		{
 			if (WebBrowserWindowPtr.IsValid())
 			{
-				WebBrowserWindowPtr.Pin()->SetTickLastFrame();
-				if (WebBrowserWindowPtr.Pin()->GetParentWindow().IsValid())
+				TSharedPtr<FWebBrowserWindow> WebBrowserWindow = WebBrowserWindowPtr.Pin();
+				WebBrowserWindow->SetTickLastFrame();
+				if (TSharedPtr<SWindow> ParentWindow = WebBrowserWindow->GetParentWindow())
 				{
-					bool ShouldSet3DBrowser = WebBrowserWindowPtr.Pin()->GetParentWindow().Get()->IsVirtualWindow();
+					bool ShouldSet3DBrowser = ParentWindow->IsVirtualWindow();
 					if (Is3DBrowser != ShouldSet3DBrowser)
 					{
 						Is3DBrowser = ShouldSet3DBrowser;
@@ -304,9 +305,9 @@ class SAppleWebBrowserWidget : public SViewport
 					// reasonable default if we need to only
 					FIntPoint ViewportSize = WebBrowserWindow->GetViewportSize();
 					
-					if(WebBrowserWindowPtr.Pin()->GetParentWindow() != nullptr)
+					if(TSharedPtr<SWindow> ParentWindow = WebBrowserWindow->GetParentWindow())
 					{
-						FVector2f WindowSize = WebBrowserWindowPtr.Pin()->GetParentWindow()->GetSizeInScreen();
+						FVector2f WindowSize = ParentWindow->GetSizeInScreen();
 						ViewportSize.X = static_cast<int32>(WindowSize.X);
 						ViewportSize.Y = static_cast<int32>(WindowSize.Y);
 #if !PLATFORM_MAC
@@ -1818,7 +1819,11 @@ void FWebBrowserWindow::SetIsDisabled(bool bValue)
 
 TSharedPtr<SWindow> FWebBrowserWindow::GetParentWindow() const
 {
-	return ParentWindow;
+	if (ParentWindow.IsValid())
+	{
+		return ParentWindow.Pin();
+	}
+	return nullptr;
 }
 
 void FWebBrowserWindow::SetParentWindow(TSharedPtr<SWindow> Window)
diff --git a/Engine/Source/Runtime/WebBrowser/Private/Apple/ApplePlatformWebBrowser.h b/Engine/Source/Runtime/WebBrowser/Private/Apple/ApplePlatformWebBrowser.h
--- a/Engine/Source/Runtime/WebBrowser/Private/Apple/ApplePlatformWebBrowser.h
+++ b/Engine/Source/Runtime/WebBrowser/Private/Apple/ApplePlatformWebBrowser.h
@@ -456,7 +456,7 @@ private:
 	/** Delegate for handling key char events not handled by the browser. */
 	FOnUnhandledKeyChar UnhandledKeyCharDelegate;
 
-	TSharedPtr<SWindow> ParentWindow;
+	TWeakPtr<SWindow> ParentWindow;
 	TWeakPtr<SDockTab> ParentDockTab;
 	FDelegateHandle ParentDockTabDelegate;
 
