Fix CEF crash when it fails to initialize

This is a CEF bug where the order of two members was inadvertently changed, causing a use-after-delete if initialization fails. This is patched in CEF itself.
Also change the check() to an ensure(), so we can still get a crash report when we detect a failed CEF initialization.
Finally, add cef3.log to the list of files we upload to the crash report service, so we can get more info about why CEF initialization might fail.

#tests verified that a forced CEF init failure doesn't trigger a crash anymore, and that cef3.log gets copied to the crash report dir as a result of the ensure
#rnx

diff --git a/Engine/Source/Runtime/WebBrowser/Private/WebBrowserSingleton.cpp b/Engine/Source/Runtime/WebBrowser/Private/WebBrowserSingleton.cpp
--- a/Engine/Source/Runtime/WebBrowser/Private/WebBrowserSingleton.cpp
+++ b/Engine/Source/Runtime/WebBrowser/Private/WebBrowserSingleton.cpp
@@ -379,12 +379,15 @@ FWebBrowserSingleton::FWebBrowserSingleton(const FWebBrowserInitSettings& WebBro
 
 		// Initialize CEF.
 		bCEFInitialized = CefInitialize(MainArgs, Settings, CEFBrowserApp.get(), nullptr);
-		check(bCEFInitialized);
+		ensure(bCEFInitialized);
 
 		// Set the thread name back to GameThread.
 		FPlatformProcess::SetThreadName(*FName(NAME_GameThread).GetPlainNameString());
 
-		DefaultCookieManager = FCefWebBrowserCookieManagerFactory::Create(CefCookieManager::GetGlobalManager(nullptr));
+		if (bCEFInitialized)
+		{
+			DefaultCookieManager = FCefWebBrowserCookieManagerFactory::Create(CefCookieManager::GetGlobalManager(nullptr));
+		}
 	}
 #elif (PLATFORM_MAC || PLATFORM_IOS) && !BUILD_EMBEDDED_APP
 	DefaultCookieManager = MakeShareable(new FAppleCookieManager());
