hotfix for crashes from apple browser upon resetting editor layout

diff --git a/Engine/Source/Runtime/WebBrowser/Private/Apple/ApplePlatformWebBrowser.cpp b/Engine/Source/Runtime/WebBrowser/Private/Apple/ApplePlatformWebBrowser.cpp
--- a/Engine/Source/Runtime/WebBrowser/Private/Apple/ApplePlatformWebBrowser.cpp
+++ b/Engine/Source/Runtime/WebBrowser/Private/Apple/ApplePlatformWebBrowser.cpp
@@ -650,59 +650,6 @@ class SAppleWebBrowserWidget : public SViewport
 	
 	void OnActiveTabChanged(TSharedPtr<SDockTab> ActiveTab, TSharedPtr<SDockTab> PreviouslyActiveTab)
 	{
-		TSharedPtr<FWebBrowserWindow> WebBrowserWindowStrongPtr = WebBrowserWindowPtr.Pin();
-		
-		if(WebBrowserWindowStrongPtr.IsValid())
-		{
-			TSharedPtr<SDockTab> CurrentTab = WebBrowserWindowStrongPtr->GetParentDockTab().Pin();
-			
-			// this should only be called when the active browser has an active parent dock tab
-			check(CurrentTab.IsValid());
-			
-			// Deal with a new major tab being foregrounded that could host our CurrentTab
-			if (ActiveTab && ActiveTab->GetTabRole() == ETabRole::MajorTab)
-			{
-				if (TSharedPtr<FTabManager> CurrentTabManager = CurrentTab->GetTabManagerPtr())
-				{
-					if (TSharedPtr<SDockTab> CurrentMajorTab = FGlobalTabmanager::Get()->GetMajorTabForTabManager(CurrentTabManager.ToSharedRef()))
-					{
-						if (ActiveTab == CurrentMajorTab)
-						{
-							if (CurrentTab->IsForeground())
-							{
-								[WebViewWrapper setVisibility : YES];
-								bForceRecomputeWindowOnPaint = true;
-							}
-						}
-					}
-				}
-			}
-
-			// Deal with a major tab being move to the background that could host our CurrentTab
-			if (PreviouslyActiveTab && PreviouslyActiveTab->GetTabRole() == ETabRole::MajorTab)
-			{
-				if (TSharedPtr<FTabManager> CurrentTabManager = CurrentTab->GetTabManagerPtr())
-				{
-					if (TSharedPtr<SDockTab> CurrentMajorTab = FGlobalTabmanager::Get()->GetMajorTabForTabManager(CurrentTabManager.ToSharedRef()))
-					{
-						if (PreviouslyActiveTab == CurrentMajorTab)
-						{
-							[WebViewWrapper setVisibility : NO ];
-						}
-					}
-				}
-			}
-
-			if(ActiveTab == CurrentTab)
-			{
-				[WebViewWrapper setVisibility: YES ];
-				bForceRecomputeWindowOnPaint = true;
-			}
-			else if(PreviouslyActiveTab == CurrentTab)
-			{
-				[WebViewWrapper setVisibility: NO ];
-			}
-		}
 	}
 	
 	~SAppleWebBrowserWidget()
@@ -1529,10 +1476,6 @@ TSharedRef<SViewport> FWebBrowserWindow::CreateWidget()
 
 void FWebBrowserWindow::OnActiveTabChanged(TSharedPtr<SDockTab> ActiveTab, TSharedPtr<SDockTab> PreviouslyActiveTab)
 {
-	if(BrowserWidget != nullptr)
-	{
-		BrowserWidget->OnActiveTabChanged(ActiveTab, PreviouslyActiveTab);
-	}
 }
 
 void FWebBrowserWindow::SetViewportSize(FIntPoint ViewportSize, FIntPoint ViewportPos)
@@ -1837,29 +1780,6 @@ void FWebBrowserWindow::SetParentWindow(TSharedPtr<SWindow> Window)
 
 void FWebBrowserWindow::SetParentDockTab(TSharedPtr<SDockTab> DockTab)
 {
-	if(DockTab == nullptr && ParentDockTabDelegate.IsValid())
-	{
-		FGlobalTabmanager::Get()->OnTabForegrounded_Unsubscribe(ParentDockTabDelegate);
-	}
-	
-	ParentDockTab = DockTab;
-	
-	if(ParentDockTab != nullptr)
-	{
-		FOnActiveTabChanged::FDelegate ActiveTabChangedDelegate;
-		
-		if(!ParentDockTabDelegate.IsValid())
-		{
-			ActiveTabChangedDelegate.BindSP(AsShared(), &FWebBrowserWindow::OnActiveTabChanged);
-			
-			ParentDockTabDelegate = FGlobalTabmanager::Get()->OnTabForegrounded_Subscribe(ActiveTabChangedDelegate);
-		}
-	}
-	else
-	{
-		// if the browser is disabled and we are setting this dock tab to nullptr, we are closing the tab
-		BrowserWidget->SetWebBrowserVisibility(!bIsDisabled);
-	}
 }
 
 TWeakPtr<SDockTab> FWebBrowserWindow::GetParentDockTab()
diff --git a/Engine/Source/Runtime/WebBrowser/Private/WebBrowserSingleton.cpp b/Engine/Source/Runtime/WebBrowser/Private/WebBrowserSingleton.cpp
--- a/Engine/Source/Runtime/WebBrowser/Private/WebBrowserSingleton.cpp
+++ b/Engine/Source/Runtime/WebBrowser/Private/WebBrowserSingleton.cpp
@@ -477,7 +477,7 @@ FWebBrowserSingleton::~FWebBrowserSingleton()
 		CefShutdown();
 	}
 	bCEFInitialized = false;
-#elif PLATFORM_IOS || PLATFORM_SPECIFIC_WEB_BROWSER || (PLATFORM_ANDROID && USE_ANDROID_JNI)
+#elif PLATFORM_APPLE || PLATFORM_SPECIFIC_WEB_BROWSER || (PLATFORM_ANDROID && USE_ANDROID_JNI)
 	{
 		FScopeLock Lock(&WindowInterfacesCS);
 		// Clear this before CefShutdown() below
@@ -843,7 +843,7 @@ bool FWebBrowserSingleton::Tick(float DeltaTime)
 		}
 	}
 
-#elif PLATFORM_IOS || PLATFORM_SPECIFIC_WEB_BROWSER || (PLATFORM_ANDROID && USE_ANDROID_JNI)
+#elif PLATFORM_APPLE || PLATFORM_SPECIFIC_WEB_BROWSER || (PLATFORM_ANDROID && USE_ANDROID_JNI)
 	FScopeLock Lock(&WindowInterfacesCS);
 	bool bIsSlateAwake = FSlateApplication::IsInitialized() && !FSlateApplication::Get().IsSlateAsleep();
 	// Remove any windows that have been deleted and check whether it's currently visible
