WebBrowser: Allow multiple CEF instances on Linux and Mac, and retry CEF initialization on Windows in case of failure

#rnx

diff --git a/Engine/Source/Runtime/WebBrowser/Private/WebBrowserSingleton.cpp b/Engine/Source/Runtime/WebBrowser/Private/WebBrowserSingleton.cpp
--- a/Engine/Source/Runtime/WebBrowser/Private/WebBrowserSingleton.cpp
+++ b/Engine/Source/Runtime/WebBrowser/Private/WebBrowserSingleton.cpp
@@ -14,6 +14,7 @@
 #include "WebBrowserLog.h"
 
 #if WITH_CEF3
+#include "CEF3Utils.h"
 #include "Misc/ScopeLock.h"
 #include "Async/Async.h"
 #include "HAL/PlatformApplicationMisc.h"
@@ -41,6 +42,7 @@ THIRD_PARTY_INCLUDES_END
 #	pragma pop_macro("OVERRIDE")
 #	if PLATFORM_WINDOWS
 #		include "Windows/HideWindowsPlatformTypes.h"
+#		include <delayimp.h>
 #	endif
 #endif
 
@@ -91,6 +93,7 @@ THIRD_PARTY_INCLUDES_END
 #	ifndef CEF3_DEFAULT_CACHE
 #		define CEF3_DEFAULT_CACHE 1
 #	endif
+#	define CEF3_INIT_MAX_RETRIES 3
 #endif
 
 FString FWebBrowserSingleton::ApplicationCacheDir() const
@@ -232,6 +235,42 @@ private:
 	struct sigaction OriginalSignalHandlers[UE_ARRAY_COUNT(PreserveSignals)];
 };
 
+static bool IsCEFLockFileCreated(const FString& LockFile)
+{
+	// Use lstat() to check the link file itself rather than the target (which isn't a valid file path)
+	struct stat FileInfo;
+	return lstat(TCHAR_TO_UTF8(*LockFile), &FileInfo) != -1 && S_ISLNK(FileInfo.st_mode);
+}
+
+static bool IsValidCEFLockFile(const FString& LockFile)
+{
+	// the lockfile on Linux/Mac is actually a symlink with a fake path formatted as <hostname>-<pid> 
+	char Buffer[128];
+	const ssize_t Len = readlink(TCHAR_TO_UTF8(*LockFile), Buffer, sizeof(Buffer));
+	if (Len <= -1 || Len >= sizeof(Buffer))
+	{
+		return false;
+	}
+	Buffer[Len] = '\0';
+	FString LockFileValue = UTF8_TO_TCHAR(Buffer);
+
+	int32 PidStartIndex = 0;
+	if (!LockFileValue.FindLastChar(TEXT('-'), PidStartIndex))
+	{
+		return false;
+	}
+	
+	uint32 Pid = 0;
+	if (!LexTryParseString(Pid, *LockFileValue.RightChop(++PidStartIndex)))
+	{
+		return false;
+	}
+
+	FProcHandle Handle = FPlatformProcess::OpenProcess(Pid);
+	ON_SCOPE_EXIT{ FPlatformProcess::CloseProc(Handle); };
+	return FPlatformProcess::IsProcRunning(Handle);
+}
+
 #endif // PLATFORM_MAC || PLATFORM_LINUX
 #endif // WITH_CEF3
 
@@ -317,8 +356,8 @@ FWebBrowserSingleton::FWebBrowserSingleton(const FWebBrowserInitSettings& WebBro
 
 #if CEF3_DEFAULT_CACHE
 		// Enable on disk cache
-		FString CachePath(FPaths::Combine(ApplicationCacheDir(), TEXT("webcache")));
-		CachePath = FPaths::ConvertRelativePathToFull(GenerateWebCacheFolderName(CachePath));
+		const FString CachePathBase(FPaths::Combine(ApplicationCacheDir(), TEXT("webcache")));
+		FString CachePath = FPaths::ConvertRelativePathToFull(GenerateWebCacheFolderName(CachePathBase));
 		CefString(&Settings.cache_path) = TCHAR_TO_WCHAR(*CachePath);
 #endif
 
@@ -378,7 +417,61 @@ FWebBrowserSingleton::FWebBrowserSingleton(const FWebBrowserInitSettings& WebBro
 #endif
 
 		// Initialize CEF.
-		bCEFInitialized = CefInitialize(MainArgs, Settings, CEFBrowserApp.get(), nullptr);
+		for (int NumRetries = 0; NumRetries <= CEF3_INIT_MAX_RETRIES; ++NumRetries)
+		{
+			bCEFInitialized = CefInitialize(MainArgs, Settings, CEFBrowserApp.get(), nullptr);
+			if (bCEFInitialized)
+			{
+				break;
+			}
+
+#if CEF3_DEFAULT_CACHE
+			const int CefExitCode = CefGetExitCode();
+			// There's a race condition that can occur when two UE instances are started close together:
+			// We check the presence of a lockfile inside the CEF cache dir to detect if that cache is already in use by another UE instance,
+			// but if the second instance checks before the first instance has created the lockfile then we'll incorrectly reuse the same cache.
+			// CEF itself handles this situation by notifying the first instance of the second launch (see FCEFBrowserApp::OnAlreadyRunningAppRelaunch)
+			// and returning CEF_RESULT_CODE_NORMAL_EXIT_PROCESS_NOTIFIED in the second instance for early exit, basically a single-app-instance pattern.
+			// If this happens, we try initializing again with another cache dir.
+			if (CefExitCode == CEF_RESULT_CODE_NORMAL_EXIT_PROCESS_NOTIFIED && NumRetries < CEF3_INIT_MAX_RETRIES)
+			{
+#if PLATFORM_WINDOWS
+				UE_LOG(LogWebBrowser, Warning, TEXT("Detected concurrent CEF initialization for cache dir %s! Retrying... (%d of %d)"), *CachePath, NumRetries + 1, CEF3_INIT_MAX_RETRIES);
+
+				// CEF/Chromium doesn't allow calling CefInitialize multiple times, even if the initialization failed,
+				// we have to fully unload the DLLs then reload them to restart from scratch.
+				// But first we must call __FUnloadDelayLoadedDLL2 to reset the delay-import address table for the symbols
+				// that were already resolved, since the DLL base address may change after the reload.
+				// Note that __FUnloadDelayLoadedDLL2 must be called here because it uses __ImageBase, so calling it from CEF3Utils wouldn't work.
+				TCHAR CEFModuleFileName[MAX_PATH] = {};
+				GetModuleFileName((HMODULE)CEF3Utils::GetCEF3ModuleHandle(), CEFModuleFileName, MAX_PATH);
+				if (!__FUnloadDelayLoadedDLL2(TCHAR_TO_UTF8(*FPaths::GetCleanFilename(CEFModuleFileName))))
+				{
+					// Don't attempt reloading if we couldn't clear the delay import address table, or it will result in crashes due to stale entries
+					UE_LOG(LogWebBrowser, Error, TEXT("Unable to cleanly reload CEF DLL, can't retry initialization!"));
+					break;
+				}
+
+				CEF3Utils::UnloadCEF3Modules();
+				CEF3Utils::LoadCEF3Modules(true);
+
+				// The cache dir returned should be a different one than previously, as the lockfile will now be detected.
+				// However the same race could be happening with this new cache dir and another UE instance starting,
+				// that's why we keep retrying as long as we get that exit code from CefInitialize.
+				CachePath = FPaths::ConvertRelativePathToFull(GenerateWebCacheFolderName(CachePathBase));
+				CefString(&Settings.cache_path) = TCHAR_TO_WCHAR(*CachePath);
+				continue;
+#else
+				// No DLL unload/reload support on Linux or Mac, we have to bail out if CefInitialize reports an already existing instance
+				UE_LOG(LogWebBrowser, Warning, TEXT("Another CEF instance is already running with cache dir %s! WebBrowser module won't be functional."), *CachePath);
+				break;
+#endif // PLATFORM_WINDOWS
+			}
+			UE_LOG(LogWebBrowser, Error, TEXT("CEF initialization for cache dir %s failed with exit code %d!"), *CachePath, CefExitCode);
+#endif // CEF3_DEFAULT_CACHE
+			break;
+		}
+		// this will upload a report to our online crash tool so we can find out why the initialization failed
 		ensure(bCEFInitialized);
 
 		// Set the thread name back to GameThread.
@@ -949,9 +1042,8 @@ FString FWebBrowserSingleton::GenerateWebCacheFolderName(const FString& InputPat
 
 #if CEF_VERSION_MAJOR >= 128
 	// CEF v128 prevents concurrent browser cache access between multiple processes.
-	// It locks the cache folder by creating a "lockfile" file in it during initialization,
+	// It locks the cache folder by creating a lockfile in it during initialization,
 	// so that other CEF initializations will fail as long as the cache folder is locked.
-	// The lockfile is then deleted when the locking CEF process ends (even in case of a crash)
 	for (uint32 CacheNumber = 0; CacheNumber < 32; ++CacheNumber) {
 		FString ConcurrentCachePath = VersionedCachePath;
 		// don't append the first _0 so we can reuse caches created as VersionedCachePath
@@ -959,13 +1051,39 @@ FString FWebBrowserSingleton::GenerateWebCacheFolderName(const FString& InputPat
 		{
 			ConcurrentCachePath.Append(FString::Format(TEXT("_{0}"), { CacheNumber }));
 		}
-		if (!FPaths::FileExists(FPaths::Combine(ConcurrentCachePath, TEXT("lockfile"))))
+#if PLATFORM_WINDOWS
+		// The lockfile is guaranteed to be deleted when the locking process ends, even in case of a crash
+		if (!FPaths::FileExists(FPaths::ConvertRelativePathToFull(FPaths::Combine(ConcurrentCachePath, TEXT("lockfile")))))
 		{
 			return ConcurrentCachePath;
 		}
+#else
+		const FString LockFile = FPaths::ConvertRelativePathToFull(FPaths::Combine(ConcurrentCachePath, TEXT("SingletonLock")));
+		if (IsCEFLockFileCreated(*LockFile))
+		{
+			// Unlike on Windows, the lockfile is *NOT* deleted if the locking process crashes,
+			// so we have to check if it's still valid. If it's not, we can use the cache dir:
+			// CEF will replace the old lockfile with a new one.
+			if (!IsValidCEFLockFile(LockFile))
+			{
+				UE_LOG(LogWebBrowser, Warning, TEXT("Found stale CEF lockfile at %s, overwriting"), *LockFile);
+				return ConcurrentCachePath;
+			}
+		}
+		else
+		{
+			// We can't re-attempt CEF initialization on Linux/Mac, so we check for the lockfile
+			// a second time after a short sleep to work around most creation races.
+			FPlatformProcess::Sleep(0.01f/*10ms*/);
+			if (!IsCEFLockFileCreated(*LockFile))
+			{
+				return ConcurrentCachePath;
+			}
+		}
+#endif // PLATFORM_WINDOWS
 	}
 	// just fall back to the base path and let CEF initialization fail if we've reached the limit
-#endif
+#endif // CEF_VERSION_MAJOR >= 128
 
 	return VersionedCachePath;
 }
