Removing unnecessary mobile renderer only metal compat checks

diff --git a/Engine/Source/Runtime/WebBrowser/Private/Apple/ApplePlatformWebBrowser.cpp b/Engine/Source/Runtime/WebBrowser/Private/Apple/ApplePlatformWebBrowser.cpp
--- a/Engine/Source/Runtime/WebBrowser/Private/Apple/ApplePlatformWebBrowser.cpp
+++ b/Engine/Source/Runtime/WebBrowser/Private/Apple/ApplePlatformWebBrowser.cpp
@@ -66,14 +66,6 @@ class SAppleWebBrowserWidget : public SViewport
 
 	void Construct(const FArguments& Args)
 	{
-		bool bSupportsMetalMRT = false;
-		GConfig->GetBool(TEXT("/Script/IOSRuntimeSettings.IOSRuntimeSettings"), TEXT("bSupportsMetalMRT"), bSupportsMetalMRT, GEngineIni);
-
-		bool bSupportsMetal = false;
-		GConfig->GetBool(TEXT("/Script/IOSRuntimeSettings.IOSRuntimeSettings"), TEXT("bSupportsMetal"), bSupportsMetal, GEngineIni);
-		// At this point we MUST be a Metal renderer.
-		check(bSupportsMetal);
-
 		WebBrowserWindowPtr = Args._WebBrowserWindow;
 		Is3DBrowser = false;
 		bAcquiredInitialFocus = false;
@@ -83,7 +75,7 @@ class SAppleWebBrowserWidget : public SViewport
 		GConfig->GetBool(TEXT("Browser"), TEXT("bEnableFloatingCloseButton"), bEnableFloatingCloseButton, GEngineIni);
 
 		WebViewWrapper = [AppleWebViewWrapper alloc];
-		[WebViewWrapper create : this userAgentApplication : Args._UserAgentApplication.GetNSString() useTransparency : Args._UseTransparency supportsMetal : bSupportsMetal supportsMetalMRT : bSupportsMetalMRT enableFloatingCloseButton : bEnableFloatingCloseButton];
+		[WebViewWrapper create : this userAgentApplication : Args._UserAgentApplication.GetNSString() useTransparency : Args._UseTransparency  enableFloatingCloseButton : bEnableFloatingCloseButton];
 
 #if !PLATFORM_TVOS
 		OnSafeFrameChangedEventHandle = FCoreDelegates::OnSafeFrameChangedEvent.AddLambda( [this]()
@@ -711,8 +703,7 @@ private:
 @synthesize GameWindowView;
 #endif
 
--(void)create:(SAppleWebBrowserWidget*)InWebBrowserWidget userAgentApplication: (NSString*)UserAgentApplication useTransparency : (bool)InUseTransparency
-supportsMetal : (bool)InSupportsMetal supportsMetalMRT : (bool)InSupportsMetalMRT enableFloatingCloseButton : (bool)bEnableFloatingCloseButton;
+-(void)create:(SAppleWebBrowserWidget*)InWebBrowserWidget userAgentApplication: (NSString*)UserAgentApplication useTransparency : (bool)InUseTransparency enableFloatingCloseButton : (bool)bEnableFloatingCloseButton;
 {
 	WebBrowserWidget = InWebBrowserWidget;
 	NextURL = nil;
@@ -721,8 +712,6 @@ supportsMetal : (bool)InSupportsMetal supportsMetalMRT : (bool)InSupportsMetalMR
 	bNeedsAddToView = true;
 	Is3DBrowser = false;
 	bVideoTextureValid = false;
-	bSupportsMetal = InSupportsMetal;
-	bSupportsMetalMRT = InSupportsMetalMRT;
 	
 #if !UE_BUILD_SHIPPING
 	BOOL Inspectable = CVarWebViewEnableInspectable.GetValueOnAnyThread()? YES : NO;
diff --git a/Engine/Source/Runtime/WebBrowser/Private/Apple/ApplePlatformWebBrowser.h b/Engine/Source/Runtime/WebBrowser/Private/Apple/ApplePlatformWebBrowser.h
--- a/Engine/Source/Runtime/WebBrowser/Private/Apple/ApplePlatformWebBrowser.h
+++ b/Engine/Source/Runtime/WebBrowser/Private/Apple/ApplePlatformWebBrowser.h
@@ -49,8 +49,6 @@ class SWebBrowserView;
 	bool bNeedsAddToView;
 	bool Is3DBrowser;
 	bool bVideoTextureValid;
-	bool bSupportsMetal;
-	bool bSupportsMetalMRT;
 }
 #if !PLATFORM_TVOS
 @property(strong) WKWebView* WebView;
@@ -68,7 +66,7 @@ class SWebBrowserView;
 @property CGRect DesiredFrame;
 
 -(void)create:(SAppleWebBrowserWidget*)InWebBrowserWidget userAgentApplication: (NSString*)UserAgentApplication useTransparency : (bool)InUseTransparency
-supportsMetal : (bool)InSupportsMetal supportsMetalMRT : (bool)InSupportsMetalMRT enableFloatingCloseButton : (bool)bEnableFloatingCloseButton;
+enableFloatingCloseButton : (bool)bEnableFloatingCloseButton;
 -(void)didRotate;
 -(void)close;
 -(void)showFloatingCloseButton:(BOOL)bShow setDraggable:(BOOL)bDraggable;
