Fixing issue with using UObject APIs in MobileJSScripting upon Engine shutdown

diff --git a/Engine/Source/Runtime/WebBrowser/Private/MobileJS/MobileJSScripting.cpp b/Engine/Source/Runtime/WebBrowser/Private/MobileJS/MobileJSScripting.cpp
--- a/Engine/Source/Runtime/WebBrowser/Private/MobileJS/MobileJSScripting.cpp
+++ b/Engine/Source/Runtime/WebBrowser/Private/MobileJS/MobileJSScripting.cpp
@@ -688,10 +688,18 @@ void FMobileJSScripting::InjectJavascript(const TSharedRef<class IWebBrowserWind
 
 	InitializeScript(Window);
 	FString Script;
-	for (auto& Item : PermanentUObjectsByName)
+	
+	if(GIsRunning)
 	{
-		Script.Append(*FString::Printf(TEXT("window.ue['%s'] = %s;"), *Item.Key.ReplaceCharWithEscapedChar(), *ConvertObject(Item.Value)));
+		// we can't try and use the ConvertObject function which uses UObject APIs 
+		// while we are shutting down, since we can't rely on the UObject API existing 
+		// in an unbroken state.
+		for (auto& Item : PermanentUObjectsByName)
+		{
+			Script.Append(*FString::Printf(TEXT("window.ue['%s'] = %s;"), *Item.Key.ReplaceCharWithEscapedChar(), *ConvertObject(Item.Value)));
+		}
 	}
+	
 	Script.Append(ScriptingPostInit);
 	Window->ExecuteJavascript(Script);
 }
