WebBrowser: Fix two major sources of CEF-related crashes reported by CrashReport

- the first one is an occasional invalid mapped memory access while Chromium is trying to parse internal memory reports. These are useless for us anyway so we effectively disable the reports by setting the frequency to a very high value
- the second one is an uncaught breakpoint exception when CEF fails to get a working GPU process. It's actually supposed to fall back to software rendering, but a bug in the offscreen rendering code caused software rendering to fail also, leaving CEF without option and giving up

The second issue required patching inside CEF, thus the new Windows binaries in this change. No change needed for Linux/Mac binaries as they're not affected by the crash (only occurs when falling back from accelerated to software rendering)

#rnx

diff --git a/Engine/Source/Runtime/WebBrowser/Private/CEF/CEFBrowserApp.cpp b/Engine/Source/Runtime/WebBrowser/Private/CEF/CEFBrowserApp.cpp
--- a/Engine/Source/Runtime/WebBrowser/Private/CEF/CEFBrowserApp.cpp
+++ b/Engine/Source/Runtime/WebBrowser/Private/CEF/CEFBrowserApp.cpp
@@ -74,6 +74,9 @@ void FCEFBrowserApp::OnBeforeCommandLineProcessing(const CefString& ProcessType,
 	CommandLine->AppendSwitch("use-mock-keychain"); // Disable the toolchain prompt on macOS.
 #endif
 
+	CommandLine->AppendSwitchWithValue("ipc-connection-timeout", "60"); // Increase the timeout from default 15s to 60s as some plugins can delay things on startup
+	CommandLine->AppendSwitchWithValue("test-memory-log-delay-in-minutes", "10000"); // Skip memory metrics reporting, it causes occasional crashes inside the memory-mapped allocator
+
 	// Uncomment these to lines to create a FULL network log from chrome, which can then be inspected using https://netlog-viewer.appspot.com/
 	//CommandLine->AppendSwitchWithValue("log-net-log", "c:\\temp\\cef_net_log.json");
 	//CommandLine->AppendSwitchWithValue("net-log-capture-mode", "IncludeCookiesAndCredentials");
