Hardening fix to avoid a crash when exiting the application.
Simply bail out of the funtion calls if slate is not initialized.

The web browser may still pump event that need to interact with the slate application but the slate application may have already been shut down.

diff --git a/Engine/Source/Runtime/WebBrowser/Private/CEF/CEFWebBrowserWindow.cpp b/Engine/Source/Runtime/WebBrowser/Private/CEF/CEFWebBrowserWindow.cpp
--- a/Engine/Source/Runtime/WebBrowser/Private/CEF/CEFWebBrowserWindow.cpp
+++ b/Engine/Source/Runtime/WebBrowser/Private/CEF/CEFWebBrowserWindow.cpp
@@ -1283,7 +1283,7 @@ bool FCEFWebBrowserWindow::OnUnhandledKeyEvent(const CefKeyEvent& CefEvent)
 			if (PreviousKeyDownEvent.IsSet())
 			{
 				bWasHandled = OnUnhandledKeyDown().IsBound() && OnUnhandledKeyDown().Execute(PreviousKeyDownEvent.GetValue());
-				if (!bWasHandled)
+				if (!bWasHandled && FSlateApplication::IsInitialized())
 				{
 					// If the keydown handler is not bound or if the handler returns false, indicating the key is unhandled, we bubble it up.
 					bIgnoreKeyDownEvent = true;
@@ -1292,7 +1292,7 @@ bool FCEFWebBrowserWindow::OnUnhandledKeyEvent(const CefKeyEvent& CefEvent)
 				}
 				PreviousKeyDownEvent.Reset();
 			}
-			else if (NativeHandle)
+			else if (NativeHandle && FSlateApplication::IsInitialized())
 			{
 				FKey const Key = FInputKeyManager::Get().GetKeyFromCodes(CefEvent.windows_key_code, 0);
 
@@ -1310,7 +1310,7 @@ bool FCEFWebBrowserWindow::OnUnhandledKeyEvent(const CefKeyEvent& CefEvent)
 			if (PreviousKeyUpEvent.IsSet())
 			{
 				bWasHandled = OnUnhandledKeyUp().IsBound() && OnUnhandledKeyUp().Execute(PreviousKeyUpEvent.GetValue());
-				if (!bWasHandled)
+				if (!bWasHandled && FSlateApplication::IsInitialized())
 				{
 					// If the keyup handler is not bound or if the handler returns false, indicating the key is unhandled, we bubble it up.
 					bIgnoreKeyUpEvent = true;
@@ -1319,7 +1319,7 @@ bool FCEFWebBrowserWindow::OnUnhandledKeyEvent(const CefKeyEvent& CefEvent)
 				}
 				PreviousKeyUpEvent.Reset();
 			}
-			else if (NativeHandle)
+			else if (NativeHandle && FSlateApplication::IsInitialized())
 			{
 				FKey const Key = FInputKeyManager::Get().GetKeyFromCodes(CefEvent.windows_key_code, 0);
 				FKeyEvent KeyEvent(Key, SlateModifiersFromCefModifiers(CefEvent), FSlateApplication::Get().GetUserIndexForKeyboard(), false, 0, CefEvent.windows_key_code);
@@ -1334,7 +1334,7 @@ bool FCEFWebBrowserWindow::OnUnhandledKeyEvent(const CefKeyEvent& CefEvent)
 			if (PreviousCharacterEvent.IsSet())
 			{
 				bWasHandled = OnUnhandledKeyChar().IsBound() && OnUnhandledKeyChar().Execute(PreviousCharacterEvent.GetValue());
-				if (!bWasHandled)
+				if (!bWasHandled && FSlateApplication::IsInitialized())
 				{
 					// If the keychar handler is not bound or if the handler returns false, indicating the key is unhandled, we bubble it up.
 					bIgnoreCharacterEvent = true;
@@ -1343,7 +1343,7 @@ bool FCEFWebBrowserWindow::OnUnhandledKeyEvent(const CefKeyEvent& CefEvent)
 				}
 				PreviousCharacterEvent.Reset();
 			}
-			else if (NativeHandle)
+			else if (NativeHandle && FSlateApplication::IsInitialized())
 			{
 				FCharacterEvent CharacterEvent(CefEvent.character, SlateModifiersFromCefModifiers(CefEvent), FSlateApplication::Get().GetUserIndexForKeyboard(), false);
 
@@ -1787,7 +1787,8 @@ void FCEFWebBrowserWindow::CloseBrowser(bool bForce, bool bBlockTillClosed)
 			Host->CloseBrowser(bForce||bBlockTillClosed);
 		}));
 
-		if (bBlockTillClosed)
+		// We can only safely pump message if slate is active
+		if (bBlockTillClosed && ensure(FSlateApplication::IsInitialized()))
 		{
 			SetIsHidden(true); // hide the window as we close it
 			float CloseWaitTimeout = 1.0f;
@@ -2168,6 +2169,11 @@ bool FCEFWebBrowserWindow::LoadCustomCEF3Cursor(cef_cursor_type_t Type)
 #define IDC_PAN_SOUTH_WEST 25855
 #define IDC_PAN_WEST 25856
 
+	if (!FSlateApplication::IsInitialized())
+	{
+		return false;
+	}
+
 	HINSTANCE CEF3ModuleHandle = (HINSTANCE)CEF3Utils::GetCEF3ModuleHandle();
 	if (CEF3ModuleHandle != nullptr)
 	{
@@ -2227,6 +2233,11 @@ bool FCEFWebBrowserWindow::LoadCustomCEF3Cursor(cef_cursor_type_t Type)
 
 bool FCEFWebBrowserWindow::OnCursorChange(CefCursorHandle CefCursor, cef_cursor_type_t Type, const CefCursorInfo& CustomCursorInfo)
 {
+	if (!FSlateApplication::IsInitialized())
+	{
+		return false;
+	}
+
 	switch (Type) {
 		// Map the basic 3 cursor types directly to Slate types on all platforms
 		case CT_NONE:
@@ -2245,7 +2256,9 @@ bool FCEFWebBrowserWindow::OnCursorChange(CefCursorHandle CefCursor, cef_cursor_
 		#if PLATFORM_WINDOWS
 				// check if we have a cursor in libcef.dll we can use
 				if (LoadCustomCEF3Cursor(Type))
+				{
 					return true;
+				}
 		#endif
 				TSharedPtr<ICursor> PlatformCursor = FSlateApplication::Get().GetPlatformCursor();
 
