WebBrowser: reduce the amount of data stored under CEF's webcache folder

A lot of the data comes from Chromium's component updater, which downloads a lot of stuff we don't want.
Turn off the component updater and other sources of background downloads.
Clear out existing webcache folders which contain that data. CEF will recreate the folder at the next run but now without all the unwanted files.

#tests verified that existing webcache folders are cleared of unwanted content after UE restart, no regressions in web-based plugins
#rnx

diff --git a/Engine/Source/Runtime/WebBrowser/Private/CEF/CEFBrowserApp.cpp b/Engine/Source/Runtime/WebBrowser/Private/CEF/CEFBrowserApp.cpp
--- a/Engine/Source/Runtime/WebBrowser/Private/CEF/CEFBrowserApp.cpp
+++ b/Engine/Source/Runtime/WebBrowser/Private/CEF/CEFBrowserApp.cpp
@@ -69,7 +69,15 @@ void FCEFBrowserApp::OnBeforeCommandLineProcessing(const CefString& ProcessType,
 
 	CommandLine->AppendSwitch("enable-begin-frame-scheduling");
 	CommandLine->AppendSwitch("disable-pinch"); // the web pages we have don't expect zoom to work right now so disable touchpad pinch zoom
-	CommandLine->AppendSwitch("disable-gpu-shader-disk-cache"); // Don't create a "GPUCache" directory when cache-path is unspecified.
+
+	// These settings reduce the amount of (unused) data downloaded by Chromium under cache_path
+	CommandLine->AppendSwitch("disable-background-networking");
+	CommandLine->AppendSwitch("disable-component-update");
+	CommandLine->AppendSwitch("disable-domain-reliability");
+	CommandLine->AppendSwitch("disable-extensions");
+	CommandLine->AppendSwitch("disable-gpu-shader-disk-cache");
+	CommandLine->AppendSwitch("disable-sync");
+
 #if PLATFORM_MAC
 	CommandLine->AppendSwitch("use-mock-keychain"); // Disable the toolchain prompt on macOS.
 #endif
diff --git a/Engine/Source/Runtime/WebBrowser/Private/WebBrowserSingleton.cpp b/Engine/Source/Runtime/WebBrowser/Private/WebBrowserSingleton.cpp
--- a/Engine/Source/Runtime/WebBrowser/Private/WebBrowserSingleton.cpp
+++ b/Engine/Source/Runtime/WebBrowser/Private/WebBrowserSingleton.cpp
@@ -355,6 +355,9 @@ FWebBrowserSingleton::FWebBrowserSingleton(const FWebBrowserInitSettings& WebBro
 		CefString(&Settings.user_agent_product) = TCHAR_TO_WCHAR(*UserAgentApplication);
 
 #if CEF3_DEFAULT_CACHE
+		// First clean up existing cache folders which contain outdated or problematic data
+		ClearOldCacheFolders(ApplicationCacheDir(), TEXT("webcache"));
+
 		// Enable on disk cache
 		const FString CachePathBase(FPaths::Combine(ApplicationCacheDir(), TEXT("webcache")));
 		FString CachePath = FPaths::ConvertRelativePathToFull(GenerateWebCacheFolderName(CachePathBase));
@@ -1093,45 +1096,32 @@ void FWebBrowserSingleton::ClearOldCacheFolders(const FString &CachePathRoot, co
 {
 #if WITH_CEF3
 	// only CEF3 currently has version dependant cache folders that may need cleanup
-	struct FDirectoryVisitor : public IPlatformFile::FDirectoryVisitor
-	{
-		const FString CachePrefix;
-		const FString CurrentCachePath;
-
-		FDirectoryVisitor(const FString &InCachePrefix, const FString &InCurrentCachePath)
-			: CachePrefix(InCachePrefix),
-			CurrentCachePath(InCurrentCachePath)
+	TArray<FString> CacheFoldersToDelete;
+	IPlatformFile::GetPlatformPhysical().IterateDirectory(*CachePathRoot,
+		[&CachePrefix, &CacheFoldersToDelete] (const TCHAR* FilenameOrDirectory, bool bIsDirectory)
 		{
-		}
-
-		virtual bool Visit(const TCHAR* FilenameOrDirectory, bool bIsDirectory) override
-		{
-			static const FString CachePrefixSearch = "/" + CachePrefix;
 			if (bIsDirectory)
 			{
 				FString DirName(FilenameOrDirectory);
-				if (DirName.Contains(CachePrefixSearch) && DirName.Equals(CurrentCachePath)==false)
+				if (FPaths::GetPathLeaf(DirName).StartsWith(CachePrefix))
 				{
-					UE_LOG(LogWebBrowser, Log, TEXT("Old Cache folder found=%s, deleting"), *DirName);
-					// BUGBUG - enable this deletion once we are happy with the new CEF version rollout
-					// Also consider adding code to preserve the previous versions folder for a while?
-					/*Async<void>(EAsyncExecution::ThreadPool, [DirName]()
-						{
-							IPlatformFile::GetPlatformPhysical().DeleteDirectoryRecursively(*DirName);
-						});*/
-
+					// Clear out cache dirs that contain unwanted folders and files automatically downloaded by Chromium's component updater but not used by UE
+					// "Crowd Deny" is one such folder, it contains files with lists of inappropriate domains (only created by CEF 128 before the component updater was disabled)
+					if (FPaths::DirectoryExists(FPaths::Combine(DirName, TEXT("Crowd Deny"))))
+					{
+						CacheFoldersToDelete.Add(DirName);
+					}
 				}
 			}
-			
 			return true;
 		}
-	};
-
-	// Enumerate the contents of the current directory
-	FDirectoryVisitor Visitor(CachePrefix, GenerateWebCacheFolderName(FPaths::Combine(CachePathRoot, CachePrefix)));
-	IPlatformFile::GetPlatformPhysical().IterateDirectory(*CachePathRoot, Visitor);
-
+	);
 
+	for (const FString& CacheFolder: CacheFoldersToDelete)
+	{
+		UE_LOG(LogWebBrowser, Log, TEXT("Deleting outdated webcache folder %s"), *CacheFolder);
+		IPlatformFile::GetPlatformPhysical().DeleteDirectoryRecursively(*CacheFolder);
+	}
 #endif
 }
 
