[WebBrowser]

Added IME setup/teardown.

diff --git a/Engine/Source/Runtime/WebBrowser/Private/SWebBrowser.cpp b/Engine/Source/Runtime/WebBrowser/Private/SWebBrowser.cpp
--- a/Engine/Source/Runtime/WebBrowser/Private/SWebBrowser.cpp
+++ b/Engine/Source/Runtime/WebBrowser/Private/SWebBrowser.cpp
@@ -5,8 +5,8 @@
 #include "Widgets/Input/SEditableTextBox.h"
 #include "Widgets/Input/SButton.h"
 #include "Widgets/Images/SThrobber.h"
-#if WITH_CEF3
 #include "Framework/Application/SlateApplication.h"
+#if WITH_CEF3
 #include "Framework/Docking/TabManager.h"
 #include "Async/Async.h"
 #include "CEF/CEFWebBrowserWindow.h"
@@ -22,6 +22,9 @@ SWebBrowser::SWebBrowser()
 
 SWebBrowser::~SWebBrowser()
 {
+#if WITH_CEF3
+	UnbindInputMethodSystem();
+#endif
 }
 
 #if WITH_CEF3
@@ -284,6 +287,8 @@ void SWebBrowser::Construct(const FArguments& InArgs, const TSharedPtr<IWebBrows
 
 
 #if WITH_CEF3
+
+
 	if (WebBrowserWindow.IsValid()) 
 	{
 		const bool bUseSmoothScrolling = 
@@ -296,7 +301,20 @@ void SWebBrowser::Construct(const FArguments& InArgs, const TSharedPtr<IWebBrows
 		
 		GetCefWebBrowserWindowChecked(WebBrowserWindow)->SetUseCapturedCefBuffer(bUseSmoothScrolling);
 	}
-#endif // WITH_CEF3
+
+
+	// Make the IME (Input Method Editor) work with HTML text input boxes on web pages,
+	// for non-ASCII languages, like "Chinese (Traditional, Taiwan)".
+	if (FSlateApplication::IsInitialized())
+	{
+		if (ITextInputMethodSystem* InputMethodSystem = FSlateApplication::Get().GetTextInputMethodSystem())
+		{
+			BindInputMethodSystem(InputMethodSystem);
+		}
+	}
+	
+	
+#endif
 }
 
 void SWebBrowser::LoadURL(FString NewURL)
@@ -530,18 +548,46 @@ void SWebBrowser::UnbindAdapter(const TSharedRef<IWebBrowserAdapter>& Adapter)
 
 void SWebBrowser::BindInputMethodSystem(ITextInputMethodSystem* TextInputMethodSystem)
 {
+	UnbindInputMethodSystem();
+
+	if (!TextInputMethodSystem || !FSlateApplication::IsInitialized())
+	{
+		return;
+	}
+
+	
 	if (BrowserView.IsValid())
 	{
 		BrowserView->BindInputMethodSystem(TextInputMethodSystem);
 	}
+
+#if WITH_CEF3
+	// During shutdown, avoids crash, ensures we unbind Input Method System while the UI is still valid.
+	InputMethodSystemSlatePreShutdownDelegateHandle = FSlateApplication::Get().OnPreShutdown().AddSP(
+		SharedThis(this), &SWebBrowser::UnbindInputMethodSystem);
+#endif
 }
 
 void SWebBrowser::UnbindInputMethodSystem()
 {
+	if (!FSlateApplication::IsInitialized())
+	{
+		return;
+	}
+	
+	
 	if (BrowserView.IsValid())
 	{
 		BrowserView->UnbindInputMethodSystem();
 	}
+
+#if WITH_CEF3
+	if (InputMethodSystemSlatePreShutdownDelegateHandle.IsValid())
+	{
+		FSlateApplication::Get().OnPreShutdown().Remove(InputMethodSystemSlatePreShutdownDelegateHandle);
+		InputMethodSystemSlatePreShutdownDelegateHandle.Reset();
+	}
+#endif
 }
 
 void SWebBrowser::SetParentWindow(TSharedPtr<SWindow> Window)
diff --git a/Engine/Source/Runtime/WebBrowser/Public/SWebBrowser.h b/Engine/Source/Runtime/WebBrowser/Public/SWebBrowser.h
--- a/Engine/Source/Runtime/WebBrowser/Public/SWebBrowser.h
+++ b/Engine/Source/Runtime/WebBrowser/Public/SWebBrowser.h
@@ -322,4 +322,7 @@ private:
 
 	/** The WebBrowserWindow itself. */
 	TSharedPtr<IWebBrowserWindow> WebBrowserWindow;
+
+	/** Pre-shutdown delegate handle, to unbind the Input Method System. */
+	FDelegateHandle InputMethodSystemSlatePreShutdownDelegateHandle;
 };
