WebBrowser: remove check() that trips in some automation use cases, and fix minor issues reported by Bughawk

#rnx

diff --git a/Engine/Source/Runtime/WebBrowser/Private/CEF/CEFWebBrowserWindowRHIHelper.cpp b/Engine/Source/Runtime/WebBrowser/Private/CEF/CEFWebBrowserWindowRHIHelper.cpp
--- a/Engine/Source/Runtime/WebBrowser/Private/CEF/CEFWebBrowserWindowRHIHelper.cpp
+++ b/Engine/Source/Runtime/WebBrowser/Private/CEF/CEFWebBrowserWindowRHIHelper.cpp
@@ -111,7 +111,7 @@ uint64_t FCEFWebBrowserWindowRHIHelper::GetRHIAdapterLuid()
 	{
 		TComPtr<IDXGIAdapter> DXGIAdapter = GetDXGIAdapterFromRHI();
 		DXGI_ADAPTER_DESC Desc = {};
-		if (SUCCEEDED(DXGIAdapter->GetDesc(&Desc)))
+		if (DXGIAdapter.IsValid() && SUCCEEDED(DXGIAdapter->GetDesc(&Desc)))
 		{
 			Luid = Desc.AdapterLuid;
 		}
@@ -131,7 +131,6 @@ uint64_t FCEFWebBrowserWindowRHIHelper::GetRHIAdapterLuid()
 FSlateUpdatableTexture* FCEFWebBrowserWindowRHIHelper::CreateSlateUpdatableTexture(const FUintPoint& TextureSize)
 {
 #if WITH_ENGINE
-	check(BUseSupportedRHIRenderer());
 #if PLATFORM_WINDOWS
 	// On Windows we need the slate texture to be shared so we can access it from the D3D device that lives in the OnAcceleratedPaint thread
 	FSlateTexture2DRHIRef* SlateTexture = new FSlateTexture2DRHIRef(TextureSize.X, TextureSize.Y, PF_B8G8R8A8, nullptr, ETextureCreateFlags::Shared | ETextureCreateFlags::RenderTargetable, true);
@@ -148,7 +147,6 @@ FSlateUpdatableTexture* FCEFWebBrowserWindowRHIHelper::CreateSlateUpdatableTextu
 void FCEFWebBrowserWindowRHIHelper::ReleaseSlateUpdatableTexture(FSlateUpdatableTexture* SlateTexture)
 {
 #if WITH_ENGINE
-	check(BUseSupportedRHIRenderer());
 #if PLATFORM_WINDOWS
 	if (!SlateTexture)
 	{
@@ -238,7 +236,6 @@ bool FCEFWebBrowserWindowRHIHelper::EnsureShareable(FSlateUpdatableTexture* Slat
 bool FCEFWebBrowserWindowRHIHelper::CopySharedTextureSync(FSlateUpdatableTexture* SlateTexture, void* SharedHandle, const FIntRect& DirtyIn)
 {
 #if WITH_ENGINE
-	check(BUseSupportedRHIRenderer());
 	FIntRect Dirty = DirtyIn;
 
 	if (!EnsureShareable(SlateTexture))
@@ -247,6 +244,11 @@ bool FCEFWebBrowserWindowRHIHelper::CopySharedTextureSync(FSlateUpdatableTexture
 	}
 
 #if PLATFORM_WINDOWS
+	if (!D3D11Device.IsValid() || !D3D11DeviceContext.IsValid())
+	{
+		return false;
+	}
+
 	TComPtr<ID3D11Device1> D3D11Device1;
 	if (HRESULT Hr = D3D11Device->QueryInterface(IID_PPV_ARGS(&D3D11Device1)); FAILED(Hr))
 	{
@@ -320,7 +322,7 @@ bool FCEFWebBrowserWindowRHIHelper::CopySharedTextureSync(FSlateUpdatableTexture
 		if (FAILED(Hr))
 		{
 			UE_LOG(LogWebBrowser, Error, TEXT("FCEFWebBrowserWindowRHIHelper::CopySharedTextureSync() - ID3D11DeviceContext::GetData  0x%x"), Hr);
-			break;
+			return false;
 		}
 
 		// We need to check for S_OK specifically as S_FALSE is also considered a success return code
