WebBrowser: fix build for targets with WITH_ENGINE=0 (e.g. SlateViewer)

#rnx

diff --git a/Engine/Source/Runtime/WebBrowser/Private/CEF/CEFWebBrowserWindowRHIHelper.cpp b/Engine/Source/Runtime/WebBrowser/Private/CEF/CEFWebBrowserWindowRHIHelper.cpp
--- a/Engine/Source/Runtime/WebBrowser/Private/CEF/CEFWebBrowserWindowRHIHelper.cpp
+++ b/Engine/Source/Runtime/WebBrowser/Private/CEF/CEFWebBrowserWindowRHIHelper.cpp
@@ -25,6 +25,7 @@ static TComPtr<IDXGIAdapter> GetDXGIAdapterFromRHI()
 {
 	TComPtr<IDXGIAdapter> DXGIAdapter;
 
+#if WITH_ENGINE
 	const ERHIInterfaceType RHIType = RHIGetInterfaceType();
 	if (RHIType == ERHIInterfaceType::D3D11)
 	{
@@ -56,6 +57,7 @@ static TComPtr<IDXGIAdapter> GetDXGIAdapterFromRHI()
 			return DXGIAdapter;
 		}
 	}
+#endif
 
 	return DXGIAdapter;
 }
