WebBrowser: Don't send the LUID to CEF if we can't extract it from RHI

Fixes webviews stuck on the spinner in SlateViewer and other programs that use `WITH_ENGINE=0`

#rnx

diff --git a/Engine/Source/Runtime/WebBrowser/Private/CEF/CEFBrowserApp.cpp b/Engine/Source/Runtime/WebBrowser/Private/CEF/CEFBrowserApp.cpp
--- a/Engine/Source/Runtime/WebBrowser/Private/CEF/CEFBrowserApp.cpp
+++ b/Engine/Source/Runtime/WebBrowser/Private/CEF/CEFBrowserApp.cpp
@@ -37,12 +37,15 @@ void FCEFBrowserApp::OnBeforeCommandLineProcessing(const CefString& ProcessType,
 
 #if PLATFORM_WINDOWS
 		// Make sure CEF uses the same adapter as UE for shared textures to work,
-		// which requires using ANGLE's D3D11 backend
-		CommandLine->AppendSwitchWithValue("use-angle", "d3d11");
+		// which requires using ANGLE's D3D11 backend and having a valid (non-zero) LUID
 		const uint64_t AdapterLuid = FCEFWebBrowserWindowRHIHelper::GetRHIAdapterLuid();
-		CommandLine->AppendSwitchWithValue("use-adapter-luid",
-			*FString::Printf(TEXT("%u,%u"), static_cast<uint32_t>(AdapterLuid >> 32), static_cast<uint32_t>(AdapterLuid & 0xFFFFFFFF)));
-#endif
+		if (AdapterLuid)
+		{
+			CommandLine->AppendSwitchWithValue("use-angle", "d3d11");
+			CommandLine->AppendSwitchWithValue("use-adapter-luid",
+				*FString::Printf(TEXT("%u,%u"), static_cast<uint32_t>(AdapterLuid >> 32), static_cast<uint32_t>(AdapterLuid & 0xFFFFFFFF)));
+		}
+#endif // PLATFORM_WINDOWS
 
 	}
 	else
