WebBrowser: Don't crash if the module is initialized before RHI is available

Add a check() instead to explain what the problem is.

#rnx

diff --git a/Engine/Source/Runtime/WebBrowser/Private/CEF/CEFWebBrowserWindowRHIHelper.cpp b/Engine/Source/Runtime/WebBrowser/Private/CEF/CEFWebBrowserWindowRHIHelper.cpp
--- a/Engine/Source/Runtime/WebBrowser/Private/CEF/CEFWebBrowserWindowRHIHelper.cpp
+++ b/Engine/Source/Runtime/WebBrowser/Private/CEF/CEFWebBrowserWindowRHIHelper.cpp
@@ -26,7 +26,7 @@ static TComPtr<IDXGIAdapter> GetDXGIAdapterFromRHI()
 	TComPtr<IDXGIAdapter> DXGIAdapter;
 
 #if WITH_ENGINE
-	const ERHIInterfaceType RHIType = RHIGetInterfaceType();
+	const ERHIInterfaceType RHIType = GDynamicRHI ? GDynamicRHI->GetInterfaceType() : ERHIInterfaceType::Hidden;
 	if (RHIType == ERHIInterfaceType::D3D11)
 	{
 		TComPtr<IDXGIDevice> DXGIDevice;
@@ -90,7 +90,8 @@ FCEFWebBrowserWindowRHIHelper::FCEFWebBrowserWindowRHIHelper()
 bool FCEFWebBrowserWindowRHIHelper::BUseSupportedRHIRenderer()
 {
 #if WITH_ENGINE
-	const ERHIInterfaceType RHIType = RHIGetInterfaceType();
+	checkf(GDynamicRHI, TEXT("WebBrowserWindow is instantiated before GDynamicRHI is initialized! CEF requires RHI resources to be available for accelerated rendering, please delay the instantiation."));
+	const ERHIInterfaceType RHIType = GDynamicRHI ? GDynamicRHI->GetInterfaceType() : ERHIInterfaceType::Hidden;
 	// We only support D3D RHIs as we rely on the OpenSharedResource1 interface
 	return (RHIType == ERHIInterfaceType::D3D11 || RHIType == ERHIInterfaceType::D3D12)
 		// Disable during automation as the GPU process often fails to initialize on build machines
@@ -107,8 +108,8 @@ uint64_t FCEFWebBrowserWindowRHIHelper::GetRHIAdapterLuid()
 #if WITH_ENGINE
 #if PLATFORM_WINDOWS
 	LUID Luid = {};
-
-	const ERHIInterfaceType RHIType = RHIGetInterfaceType();
+	checkf(GDynamicRHI, TEXT("WebBrowserModule is instantiated before GDynamicRHI is initialized! CEF requires RHI resources to be available for accelerated rendering, please delay the instantiation."));
+	const ERHIInterfaceType RHIType = GDynamicRHI ? GDynamicRHI->GetInterfaceType() : ERHIInterfaceType::Hidden;
 	if (RHIType == ERHIInterfaceType::D3D11)
 	{
 		TComPtr<IDXGIAdapter> DXGIAdapter = GetDXGIAdapterFromRHI();
@@ -192,7 +193,7 @@ bool FCEFWebBrowserWindowRHIHelper::EnsureShareable(FSlateUpdatableTexture* Slat
 	HANDLE SharedHandle = nullptr;
 
 	// We add SharedHandle to the map even if we couldn't extract it, so we don't continuously try and report failures
-	const ERHIInterfaceType RHIType = RHIGetInterfaceType();
+	const ERHIInterfaceType RHIType = GDynamicRHI ? GDynamicRHI->GetInterfaceType() : ERHIInterfaceType::Hidden;
 	if (RHIType == ERHIInterfaceType::D3D11)
 	{
 		TComPtr<IDXGIResource> DXGIResource;
@@ -280,7 +281,7 @@ bool FCEFWebBrowserWindowRHIHelper::CopySharedTextureSync(FSlateUpdatableTexture
 	}
 
 	TComPtr<ID3D11Texture2D> DestTexture;
-	const ERHIInterfaceType RHIType = RHIGetInterfaceType();
+	const ERHIInterfaceType RHIType = GDynamicRHI ? GDynamicRHI->GetInterfaceType() : ERHIInterfaceType::Hidden;
 	if (RHIType == ERHIInterfaceType::D3D11)
 	{
 		if (HRESULT Hr = D3D11Device->OpenSharedResource((HANDLE)SlateTextureHandles.FindRef(SlateTexture).Get(), IID_PPV_ARGS(&DestTexture)); FAILED(Hr))
