Fix IOS slate webview not applying contentscale correctly

#rnx
#tests Built locally and tested on iphone 16

diff --git a/Engine/Source/Runtime/WebBrowser/Private/Apple/ApplePlatformWebBrowser.cpp b/Engine/Source/Runtime/WebBrowser/Private/Apple/ApplePlatformWebBrowser.cpp
--- a/Engine/Source/Runtime/WebBrowser/Private/Apple/ApplePlatformWebBrowser.cpp
+++ b/Engine/Source/Runtime/WebBrowser/Private/Apple/ApplePlatformWebBrowser.cpp
@@ -168,17 +168,12 @@ class SAppleWebBrowserWidget : public SViewport
 		{
 			UIView* View = [IOSAppDelegate GetDelegate].IOSView;
 			CGFloat ContentScaleFactor = View.contentScaleFactor;
-				
-			FVector2D PosVector2D = WindowPos * ContentScaleFactor;
-			
-			WindowPos /= ContentScaleFactor;
-			WindowSize /= ContentScaleFactor;
 			
-			NewFrame.size.width = WindowSize.X;
-			NewFrame.size.height = WindowSize.Y;
+			NewFrame.size.width = FMath::RoundToInt(WindowSize.X / ContentScaleFactor);
+			NewFrame.size.height = FMath::RoundToInt(WindowSize.Y / ContentScaleFactor);
 			
-			NewFrame.origin.x = WindowPos.X;
-			NewFrame.origin.y = WindowPos.Y;
+			NewFrame.origin.x = FMath::RoundToInt(WindowPos.X / ContentScaleFactor);
+			NewFrame.origin.y = FMath::RoundToInt(WindowPos.Y / ContentScaleFactor);
 		}
 #endif
 
