WebBrowser: Increase the timeout for initial connection to the EpicWebHelper processes

#rnx

diff --git a/Engine/Source/Runtime/WebBrowser/Private/CEF/CEFBrowserApp.cpp b/Engine/Source/Runtime/WebBrowser/Private/CEF/CEFBrowserApp.cpp
--- a/Engine/Source/Runtime/WebBrowser/Private/CEF/CEFBrowserApp.cpp
+++ b/Engine/Source/Runtime/WebBrowser/Private/CEF/CEFBrowserApp.cpp
@@ -77,6 +77,8 @@ void FCEFBrowserApp::OnBeforeCommandLineProcessing(const CefString& ProcessType,
 	CommandLine->AppendSwitch("use-mock-keychain"); // Disable the toolchain prompt on macOS.
 #endif
 
+	CommandLine->AppendSwitchWithValue("ipc-connection-timeout", "60"); // Increase the timeout from default 15s to 60s as some plugins can delay things on startup
+
 	// Uncomment these to lines to create a FULL network log from chrome, which can then be inspected using https://netlog-viewer.appspot.com/
 	//CommandLine->AppendSwitchWithValue("log-net-log", "c:\\temp\\cef_net_log.json");
 	//CommandLine->AppendSwitchWithValue("net-log-capture-mode", "IncludeCookiesAndCredentials");
