WebBrowser: fix memory leak when invoking reflected UE functions from JavaScript

The Params memory itself was freed, but the destructors for the parameters inside there were not called,
resulting in leaked memory in particular for FString parameters, for which the buffer wasn't freed.

Removed the DestroyStruct call from the CEF impl which confusingly implied it destroyed the parameters,
but didn't actually do anything (because the params don't have the DestructorLink set).

PR #13404

#rnx

diff --git a/Engine/Source/Runtime/WebBrowser/Private/CEF/CEFJSScripting.cpp b/Engine/Source/Runtime/WebBrowser/Private/CEF/CEFJSScripting.cpp
--- a/Engine/Source/Runtime/WebBrowser/Private/CEF/CEFJSScripting.cpp
+++ b/Engine/Source/Runtime/WebBrowser/Private/CEF/CEFJSScripting.cpp
@@ -325,7 +325,17 @@ bool FCEFJSScripting::HandleExecuteUObjectMethodMessage(CefRefPtr<CefListValue>
 
 	if (Params)
 	{
-		Function->DestroyStruct(Params);
+		// Invoke the parameters destructors as needed before freeing the Params memory itself
+		for (TFieldIterator<FProperty> It(Function); It && It->HasAnyPropertyFlags(CPF_Parm); ++It)
+		{
+			if (It->HasAnyPropertyFlags(CPF_NoDestructor) || !It->IsInContainer(Function->GetStructureSize()))
+			{
+				continue;
+			}
+
+			It->DestroyValue_InContainer(Params);
+		}
+
 		FMemory::Free(Params);
 		Params = nullptr;
 	}
diff --git a/Engine/Source/Runtime/WebBrowser/Private/MobileJS/MobileJSScripting.cpp b/Engine/Source/Runtime/WebBrowser/Private/MobileJS/MobileJSScripting.cpp
--- a/Engine/Source/Runtime/WebBrowser/Private/MobileJS/MobileJSScripting.cpp
+++ b/Engine/Source/Runtime/WebBrowser/Private/MobileJS/MobileJSScripting.cpp
@@ -617,6 +617,21 @@ bool FMobileJSScripting::HandleExecuteUObjectMethodMessage(const TArray<FString>
 			InvokeJSFunction(ResultCallbackId, 0, nullptr, false);
 		}
 	}
+
+	if (Params.Num() > 0)
+	{
+		// Invoke the parameters destructors as needed before Params itself gets released
+		for (TFieldIterator<FProperty> It(Function); It && It->HasAnyPropertyFlags(CPF_Parm); ++It)
+		{
+			if (It->HasAnyPropertyFlags(CPF_NoDestructor) || !It->IsInContainer(Function->GetStructureSize()))
+			{
+				continue;
+			}
+
+			It->DestroyValue_InContainer(Params.GetData());
+		}
+	}
+	
 	return true;
 }
 
diff --git a/Engine/Source/Runtime/WebBrowser/Private/Native/NativeJSScripting.cpp b/Engine/Source/Runtime/WebBrowser/Private/Native/NativeJSScripting.cpp
--- a/Engine/Source/Runtime/WebBrowser/Private/Native/NativeJSScripting.cpp
+++ b/Engine/Source/Runtime/WebBrowser/Private/Native/NativeJSScripting.cpp
@@ -478,6 +478,21 @@ bool FNativeJSScripting::HandleExecuteUObjectMethodMessage(const TArray<FString>
 			InvokeJSFunction(ResultCallbackId, 0, nullptr, false);
 		}
 	}
+
+	if (Params.Num() > 0)
+	{
+		// Invoke the parameters destructors as needed before Params itself gets released
+		for (TFieldIterator<FProperty> It(Function); It && It->HasAnyPropertyFlags(CPF_Parm); ++It)
+		{
+			if (It->HasAnyPropertyFlags(CPF_NoDestructor) || !It->IsInContainer(Function->GetStructureSize()))
+			{
+				continue;
+			}
+
+			It->DestroyValue_InContainer(Params.GetData());
+		}
+	}
+	
 	return true;
 }
 
