WebBrowser: forward focus to the BrowserView when focusing programmatically

This allows focusing the webview via `FSlateApplication::SetKeyboardFocus`. Previously the event was just unhandled, making programmatic focus impossible.

PR #9462

#rnx

diff --git a/Engine/Source/Runtime/WebBrowser/Private/SWebBrowser.cpp b/Engine/Source/Runtime/WebBrowser/Private/SWebBrowser.cpp
--- a/Engine/Source/Runtime/WebBrowser/Private/SWebBrowser.cpp
+++ b/Engine/Source/Runtime/WebBrowser/Private/SWebBrowser.cpp
@@ -598,5 +598,16 @@ void SWebBrowser::SetParentWindow(TSharedPtr<SWindow> Window)
 	}
 }
 
+FReply SWebBrowser::OnFocusReceived(const FGeometry& MyGeometry, const FFocusEvent& InFocusEvent)
+{
+	if (BrowserView.IsValid() && InFocusEvent.GetCause() == EFocusCause::SetDirectly)
+	{
+		FReply Reply = FReply::Handled();
+		Reply.SetUserFocus(BrowserView.ToSharedRef(), InFocusEvent.GetCause(), true);
+		return Reply;
+	}
+
+	return SCompoundWidget::OnFocusReceived(MyGeometry, InFocusEvent);
+}
 
 #undef LOCTEXT_NAMESPACE
diff --git a/Engine/Source/Runtime/WebBrowser/Public/SWebBrowser.h b/Engine/Source/Runtime/WebBrowser/Public/SWebBrowser.h
--- a/Engine/Source/Runtime/WebBrowser/Public/SWebBrowser.h
+++ b/Engine/Source/Runtime/WebBrowser/Public/SWebBrowser.h
@@ -137,6 +137,7 @@ public:
 	WEBBROWSER_API ~SWebBrowser();
 
 	virtual bool SupportsKeyboardFocus() const override {return true;}
+	virtual FReply OnFocusReceived(const FGeometry& MyGeometry, const FFocusEvent& InFocusEvent) override;
 
 	virtual void Tick(const FGeometry& AllottedGeometry, const double InCurrentTime, const float InDeltaTime) override;
 	virtual int32 OnPaint(const FPaintArgs& Args, const FGeometry& AllottedGeometry, const FSlateRect& MyCullingRect, FSlateWindowElementList& OutDrawElements, int32 LayerId, const FWidgetStyle& InWidgetStyle, bool bParentEnabled) const override;
