making ApplePlatformWebBrowser Slate visibility aware without plugin intervention

#Synced-CL Unknown

diff --git a/Engine/Plugins/Bridge/Source/Bridge/Private/UI/BridgeUIManager.cpp b/Engine/Plugins/Bridge/Source/Bridge/Private/UI/BridgeUIManager.cpp
--- a/Engine/Plugins/Bridge/Source/Bridge/Private/UI/BridgeUIManager.cpp
+++ b/Engine/Plugins/Bridge/Source/Bridge/Private/UI/BridgeUIManager.cpp
@@ -339,7 +339,6 @@ TSharedRef<SDockTab> FBridgeUIManagerImpl::CreateBridgeTab(const FSpawnTabArgs&
 			FBridgeUIManager::Instance->LocalBrowserDock = NULL;
 			if (FBridgeUIManager::Instance->WebBrowserWidget.IsValid())
 			{
-				FBridgeUIManager::Instance->Browser->SetParentDockTab(nullptr);
 				FBridgeUIManager::Instance->WebBrowserWidget.Reset();
 				FBridgeUIManager::Instance->Browser.Reset();
 			}
@@ -372,13 +371,10 @@ TSharedRef<SDockTab> FBridgeUIManagerImpl::CreateBridgeTab(const FSpawnTabArgs&
 
 	if (LocalBrowserDock.IsValid())
 	{
-		Browser->SetParentDockTab(LocalBrowserDock);
 		return LocalBrowserDock.ToSharedRef();
 	}
 	
 	SAssignNew(LocalBrowserDock, SDockTab);
-	
-	Browser->SetParentDockTab(LocalBrowserDock);
 
 	return LocalBrowserDock.ToSharedRef();
 }
diff --git a/Engine/Plugins/Experimental/AIAssistant/Source/AIAssistant/Private/AIAssistantWebBrowser.cpp b/Engine/Plugins/Experimental/AIAssistant/Source/AIAssistant/Private/AIAssistantWebBrowser.cpp
--- a/Engine/Plugins/Experimental/AIAssistant/Source/AIAssistant/Private/AIAssistantWebBrowser.cpp
+++ b/Engine/Plugins/Experimental/AIAssistant/Source/AIAssistant/Private/AIAssistantWebBrowser.cpp
@@ -60,7 +60,6 @@ void SAIAssistantWebBrowser::Construct(const FArguments& InArgs, const TSharedRe
 		WindowSettings.BrowserFrameRate = 60;
 
 		WebBrowserWindow = WebBrowserSingleton->CreateBrowserWindow(WindowSettings);
-		WebBrowserWindow->SetParentDockTab(InParentTab);
 	}
 	else
 	{
diff --git a/Engine/Plugins/Fab/Source/Fab/Private/FabBrowser.cpp b/Engine/Plugins/Fab/Source/Fab/Private/FabBrowser.cpp
--- a/Engine/Plugins/Fab/Source/Fab/Private/FabBrowser.cpp
+++ b/Engine/Plugins/Fab/Source/Fab/Private/FabBrowser.cpp
@@ -371,8 +371,6 @@ TSharedRef<SDockTab> FFabBrowser::OpenTab(const FSpawnTabArgs&)
 #if PLATFORM_MAC && !WITH_CEF3
 	FSlateApplication::Get().ToggleDisableLastDragOnDragEnter(true);
 #endif
-	
-	WebBrowserWindow->SetParentDockTab(DockTab);
 
 	return DockTab.ToSharedRef();
 }
@@ -394,8 +392,10 @@ void FFabBrowser::OnPluginTabClosed(TSharedRef<SDockTab> InParentTab)
 	);
 
 #if PLATFORM_MAC && !WITH_CEF3
-	WebBrowserWindow->SetIsDisabled(true);
-	WebBrowserWindow->SetParentDockTab(nullptr);
+	if(WebBrowserWindow.IsValid())
+	{
+		WebBrowserWindow->SetIsDisabled(true);
+	}
 #endif
 	
 	AsyncTask(ENamedThreads::AnyBackgroundThreadNormalTask, []()
diff --git a/Engine/Source/Editor/MainFrame/Private/MainFrameModule.cpp b/Engine/Source/Editor/MainFrame/Private/MainFrameModule.cpp
--- a/Engine/Source/Editor/MainFrame/Private/MainFrameModule.cpp
+++ b/Engine/Source/Editor/MainFrame/Private/MainFrameModule.cpp
@@ -1207,10 +1207,6 @@ TSharedRef<SWidget> FMainFrameModule::CreateHomeScreenWidget(const TSharedPtr<SD
 			return SNullWidget::NullWidget;
 		}
 		
-		if (InParentDockTab)
-		{
-			WebBrowserWindow->SetParentDockTab(InParentDockTab);
-		}
 		WebBrowserSingleton->SetDevToolsShortcutEnabled(true);
 		WebBrowserWindow->OnCreateWindow().BindLambda([](const TWeakPtr<IWebBrowserWindow>& InNewBrowserWindowWeak, const TWeakPtr<IWebBrowserPopupFeatures>& PopupFeatures)
 		{
diff --git a/Engine/Source/Runtime/WebBrowser/Private/Apple/ApplePlatformWebBrowser.cpp b/Engine/Source/Runtime/WebBrowser/Private/Apple/ApplePlatformWebBrowser.cpp
--- a/Engine/Source/Runtime/WebBrowser/Private/Apple/ApplePlatformWebBrowser.cpp
+++ b/Engine/Source/Runtime/WebBrowser/Private/Apple/ApplePlatformWebBrowser.cpp
@@ -138,6 +138,14 @@ class SAppleWebBrowserWidget : public SViewport
 	void UpdateFrameWithViewport(FIntPoint WindowSize, FIntPoint WindowPos) {
 		CGRect NewFrame;
 #if PLATFORM_MAC
+		TSharedPtr<FWebBrowserWindow> BrowserWindow = WebBrowserWindowPtr.Pin();
+		
+		if(!BrowserWindow.IsValid() || !BrowserWindow->GetParentWindow().IsValid() ||
+		   !BrowserWindow->GetParentWindow()->GetNativeWindow().IsValid())
+		{
+			return;
+		}
+		
 		TSharedPtr<FGenericWindow> EnclosedWindow = WebBrowserWindowPtr.Pin()->GetParentWindow()->GetNativeWindow();
 		
 		if(!Is3DBrowser) 
@@ -643,63 +651,6 @@ class SAppleWebBrowserWidget : public SViewport
 		return WebViewWrapper != nil && [WebViewWrapper inKeyWindow];
 	}
 	
-	void OnActiveTabChanged(TSharedPtr<SDockTab> ActiveTab, TSharedPtr<SDockTab> PreviouslyActiveTab)
-	{
-		TSharedPtr<FWebBrowserWindow> WebBrowserWindowStrongPtr = WebBrowserWindowPtr.Pin();
-		
-		if(WebBrowserWindowStrongPtr.IsValid())
-		{
-			TSharedPtr<SDockTab> CurrentTab = WebBrowserWindowStrongPtr->GetParentDockTab().Pin();
-			
-			// this should only be called when the active browser has an active parent dock tab
-			check(CurrentTab.IsValid());
-			
-			// Deal with a new major tab being foregrounded that could host our CurrentTab
-			if (ActiveTab && ActiveTab->GetTabRole() == ETabRole::MajorTab)
-			{
-				if (TSharedPtr<FTabManager> CurrentTabManager = CurrentTab->GetTabManagerPtr())
-				{
-					if (TSharedPtr<SDockTab> CurrentMajorTab = FGlobalTabmanager::Get()->GetMajorTabForTabManager(CurrentTabManager.ToSharedRef()))
-					{
-						if (ActiveTab == CurrentMajorTab)
-						{
-							if (CurrentTab->IsForeground())
-							{
-								[WebViewWrapper setVisibility : YES];
-								bForceRecomputeWindowOnPaint = true;
-							}
-						}
-					}
-				}
-			}
-
-			// Deal with a major tab being move to the background that could host our CurrentTab
-			if (PreviouslyActiveTab && PreviouslyActiveTab->GetTabRole() == ETabRole::MajorTab)
-			{
-				if (TSharedPtr<FTabManager> CurrentTabManager = CurrentTab->GetTabManagerPtr())
-				{
-					if (TSharedPtr<SDockTab> CurrentMajorTab = FGlobalTabmanager::Get()->GetMajorTabForTabManager(CurrentTabManager.ToSharedRef()))
-					{
-						if (PreviouslyActiveTab == CurrentMajorTab)
-						{
-							[WebViewWrapper setVisibility : NO ];
-						}
-					}
-				}
-			}
-
-			if(ActiveTab == CurrentTab)
-			{
-				[WebViewWrapper setVisibility: YES ];
-				bForceRecomputeWindowOnPaint = true;
-			}
-			else if(PreviouslyActiveTab == CurrentTab)
-			{
-				[WebViewWrapper setVisibility: NO ];
-			}
-		}
-	}
-	
 	~SAppleWebBrowserWidget()
 	{
 		Close();
@@ -1522,14 +1473,6 @@ TSharedRef<SViewport> FWebBrowserWindow::CreateWidget()
 	return BrowserWidgetRef;
 }
 
-void FWebBrowserWindow::OnActiveTabChanged(TSharedPtr<SDockTab> ActiveTab, TSharedPtr<SDockTab> PreviouslyActiveTab)
-{
-	if(BrowserWidget != nullptr)
-	{
-		BrowserWidget->OnActiveTabChanged(ActiveTab, PreviouslyActiveTab);
-	}
-}
-
 void FWebBrowserWindow::SetViewportSize(FIntPoint ViewportSize, FIntPoint ViewportPos)
 {
 	WindowSize = ViewportSize;
@@ -1830,38 +1773,6 @@ void FWebBrowserWindow::SetParentWindow(TSharedPtr<SWindow> Window)
 	ParentWindow = Window;
 }
 
-void FWebBrowserWindow::SetParentDockTab(TSharedPtr<SDockTab> DockTab)
-{
-	if(DockTab == nullptr && ParentDockTabDelegate.IsValid())
-	{
-		FGlobalTabmanager::Get()->OnTabForegrounded_Unsubscribe(ParentDockTabDelegate);
-	}
-	
-	ParentDockTab = DockTab;
-	
-	if(ParentDockTab != nullptr)
-	{
-		FOnActiveTabChanged::FDelegate ActiveTabChangedDelegate;
-		
-		if(!ParentDockTabDelegate.IsValid())
-		{
-			ActiveTabChangedDelegate.BindSP(AsShared(), &FWebBrowserWindow::OnActiveTabChanged);
-			
-			ParentDockTabDelegate = FGlobalTabmanager::Get()->OnTabForegrounded_Subscribe(ActiveTabChangedDelegate);
-		}
-	}
-	else
-	{
-		// if the browser is disabled and we are setting this dock tab to nullptr, we are closing the tab
-		BrowserWidget->SetWebBrowserVisibility(!bIsDisabled);
-	}
-}
-
-TWeakPtr<SDockTab> FWebBrowserWindow::GetParentDockTab()
-{
-	return ParentDockTab;
-}
-
 void FWebBrowserWindow::ShowFloatingCloseButton(bool bShow, bool bDraggable)
 {
 	if (BrowserWidget)
diff --git a/Engine/Source/Runtime/WebBrowser/Private/Apple/ApplePlatformWebBrowser.h b/Engine/Source/Runtime/WebBrowser/Private/Apple/ApplePlatformWebBrowser.h
--- a/Engine/Source/Runtime/WebBrowser/Private/Apple/ApplePlatformWebBrowser.h
+++ b/Engine/Source/Runtime/WebBrowser/Private/Apple/ApplePlatformWebBrowser.h
@@ -181,7 +181,6 @@ public:
 	virtual void SetIsDisabled(bool bValue) override;
 	virtual TSharedPtr<SWindow> GetParentWindow() const override;
 	virtual void SetParentWindow(TSharedPtr<SWindow> Window) override;
-	virtual void SetParentDockTab(TSharedPtr<SDockTab> DockTab) override;
 	virtual void ShowFloatingCloseButton(bool bShow, bool bDraggable) override;
 
 	// TODO: None of these events are actually called
@@ -351,10 +350,6 @@ public:
 		OnTitleChanged().Broadcast(Title);
 	}
 	
-	void OnActiveTabChanged(TSharedPtr<SDockTab> ActiveTab, TSharedPtr<SDockTab> PreviouslyActiveTab);
-	
-	TWeakPtr<SDockTab> GetParentDockTab();
-
 private:
 
 	TOptional<FSlateRenderTransform> GetWebBrowserRenderTransform() const;
@@ -457,7 +452,6 @@ private:
 	FOnUnhandledKeyChar UnhandledKeyCharDelegate;
 
 	TWeakPtr<SWindow> ParentWindow;
-	TWeakPtr<SDockTab> ParentDockTab;
 	FDelegateHandle ParentDockTabDelegate;
 
 	FIntPoint WindowSize;
diff --git a/Engine/Source/Runtime/WebBrowser/Private/WebBrowserSingleton.cpp b/Engine/Source/Runtime/WebBrowser/Private/WebBrowserSingleton.cpp
--- a/Engine/Source/Runtime/WebBrowser/Private/WebBrowserSingleton.cpp
+++ b/Engine/Source/Runtime/WebBrowser/Private/WebBrowserSingleton.cpp
@@ -837,7 +837,7 @@ bool FWebBrowserSingleton::Tick(float DeltaTime)
 		}
 	}
 
-#elif PLATFORM_IOS || PLATFORM_SPECIFIC_WEB_BROWSER || (PLATFORM_ANDROID && USE_ANDROID_JNI)
+#elif PLATFORM_APPLE || PLATFORM_SPECIFIC_WEB_BROWSER || (PLATFORM_ANDROID && USE_ANDROID_JNI)
 	FScopeLock Lock(&WindowInterfacesCS);
 	bool bIsSlateAwake = FSlateApplication::IsInitialized() && !FSlateApplication::Get().IsSlateAsleep();
 	// Remove any windows that have been deleted and check whether it's currently visible
diff --git a/Engine/Source/Runtime/WebBrowser/Public/IWebBrowserWindow.h b/Engine/Source/Runtime/WebBrowser/Public/IWebBrowserWindow.h
--- a/Engine/Source/Runtime/WebBrowser/Public/IWebBrowserWindow.h
+++ b/Engine/Source/Runtime/WebBrowser/Public/IWebBrowserWindow.h
@@ -374,8 +374,6 @@ public:
 	virtual void ShowFloatingCloseButton(bool bShow, bool bDraggable) = 0;
 
 	virtual void CheckTickActivity() {};
-	
-	virtual void SetParentDockTab(TSharedPtr<class SDockTab> DockTab) {};
 
 public:
 
