WebBrowser: Fix CEF webview scaling in native window mode

CEF v128 now applies DPI scaling to the coordinates set for the child window on Windows, so we have to cancel it using the DPI scale factor for the monitor that the window is on.

#tests verified that the webview keeps the right size with different monitor scaling settings
#rnx

#Synced-CL 48762613

diff --git a/Engine/Source/Runtime/WebBrowser/Private/WebBrowserSingleton.cpp b/Engine/Source/Runtime/WebBrowser/Private/WebBrowserSingleton.cpp
--- a/Engine/Source/Runtime/WebBrowser/Private/WebBrowserSingleton.cpp
+++ b/Engine/Source/Runtime/WebBrowser/Private/WebBrowserSingleton.cpp
@@ -567,7 +567,17 @@ TSharedPtr<IWebBrowserWindow> FWebBrowserSingleton::CreateBrowserWindow(const FC
 #if CEF_VERSION_MAJOR < 128
 			WindowInfo.SetAsChild((CefWindowHandle)WindowSettings.OSWindowHandle, ClientRect);
 #else
-			WindowInfo.SetAsChild((CefWindowHandle)WindowSettings.OSWindowHandle, { ClientRect.left, ClientRect.top, ClientRect.right - ClientRect.left, ClientRect.bottom - ClientRect.top });
+			RECT WindowRect = {};
+			float DPIScale = 1.0f;
+			if (!GetWindowRect((HWND)WindowSettings.OSWindowHandle, &WindowRect))
+			{
+				UE_LOG(LogWebBrowser, Error, TEXT("Failed to get window rect"));
+			}
+			else
+			{
+				DPIScale = FPlatformApplicationMisc::GetDPIScaleFactorAtPoint(WindowRect.left, WindowRect.top);
+			}
+			WindowInfo.SetAsChild((CefWindowHandle)WindowSettings.OSWindowHandle, { 0, 0, static_cast<int>((ClientRect.right - ClientRect.left) / DPIScale), static_cast<int>((ClientRect.bottom - ClientRect.top) / DPIScale) });
 #endif
 		}
 		else
