fixing missing dependency for moving assets from WebBrowserWidget to WebBrowser Runtime module

#Synced-CL Unknown

diff --git a/Engine/Source/Runtime/WebBrowser/WebBrowser.Build.cs b/Engine/Source/Runtime/WebBrowser/WebBrowser.Build.cs
--- a/Engine/Source/Runtime/WebBrowser/WebBrowser.Build.cs
+++ b/Engine/Source/Runtime/WebBrowser/WebBrowser.Build.cs
@@ -17,6 +17,7 @@ public class WebBrowser : ModuleRules
 				"InputCore",
 				"Serialization",
 				"HTTP",
+				"Engine"
 			}
 		);
 
