[WebBrowser]

Remove blanket Engine dependency for WebBrowser and use it when Target.bCompileAgainstEngine

#Synced-CL Unknown

diff --git a/Engine/Source/Runtime/WebBrowser/Private/WebBrowserSingleton.cpp b/Engine/Source/Runtime/WebBrowser/Private/WebBrowserSingleton.cpp
--- a/Engine/Source/Runtime/WebBrowser/Private/WebBrowserSingleton.cpp
+++ b/Engine/Source/Runtime/WebBrowser/Private/WebBrowserSingleton.cpp
@@ -247,8 +247,10 @@ FWebBrowserSingleton::FWebBrowserSingleton(const FWebBrowserInitSettings& WebBro
 #if WITH_CEF3
 	, bCEFInitialized(false)
 #endif
+#if WITH_ENGINE
 	, DefaultMaterial(FSoftObjectPath(FString(TEXT("/Engine/WebBrowser/WebTexture_M.WebTexture_M"))))
 	, DefaultTranslucentMaterial(FSoftObjectPath(FString(TEXT("/Engine/WebBrowser/WebTexture_TM.WebTexture_TM"))))
+#endif
 {
 #if WITH_CEF3
 	
@@ -392,8 +394,10 @@ FWebBrowserSingleton::FWebBrowserSingleton(const FWebBrowserInitSettings& WebBro
 	DefaultCookieManager = MakeShareable(new FAndroidCookieManager());
 #endif
 	
+#if WITH_ENGINE
 	DefaultMaterial.LoadSynchronous();
 	DefaultTranslucentMaterial.LoadSynchronous();
+#endif
 }
 
 
diff --git a/Engine/Source/Runtime/WebBrowser/Private/WebBrowserSingleton.h b/Engine/Source/Runtime/WebBrowser/Private/WebBrowserSingleton.h
--- a/Engine/Source/Runtime/WebBrowser/Private/WebBrowserSingleton.h
+++ b/Engine/Source/Runtime/WebBrowser/Private/WebBrowserSingleton.h
@@ -6,7 +6,9 @@
 #include "Containers/Ticker.h"
 #include "IWebBrowserSingleton.h"
 
+#if WITH_ENGINE
 #include "Materials/Material.h"
+#endif
 
 #if WITH_CEF3
 #if PLATFORM_WINDOWS
@@ -123,27 +125,38 @@ public:
 	/** Set a reference to UWebBrowser's default material*/
 	virtual void SetDefaultMaterial(UMaterialInterface* InDefaultMaterial) override
 	{
+#if WITH_EDITOR
 		DefaultMaterial = InDefaultMaterial;
+#endif
 	}
 
 	/** Set a reference to UWebBrowser's translucent material*/
 	virtual void SetDefaultTranslucentMaterial(UMaterialInterface* InDefaultMaterial) override
 	{
+#if WITH_EDITOR
 		DefaultTranslucentMaterial = InDefaultMaterial;
+#endif
 	}
 
 	/** Get a reference to UWebBrowser's default material*/
 	virtual UMaterialInterface* GetDefaultMaterial() override
 	{
+#if WITH_EDITOR
 		return DefaultMaterial.Get();
+#else
+		return nullptr;
+#endif
 	}
 
 	/** Get a reference to UWebBrowser's translucent material*/
 	virtual UMaterialInterface* GetDefaultTranslucentMaterial() override
 	{
+#if WITH_EDITOR
 		return DefaultTranslucentMaterial.Get();
+#else
+		return nullptr;
+#endif
 	}
-
 public:
 
 	// FTSTickerObjectBase Interface
@@ -207,11 +220,13 @@ private:
 	bool bCEFInitialized;
 #endif
 
+#if WITH_ENGINE
 	/** Reference to UWebBrowser's default material*/
 	TSoftObjectPtr<UMaterialInterface> DefaultMaterial;
 
 	/** Reference to UWebBrowser's translucent material*/
 	TSoftObjectPtr<UMaterialInterface> DefaultTranslucentMaterial;
+#endif
 
 	/* Time when our Tick() was last called. */
 	double PreviousTickTimeSeconds = 0.0;
diff --git a/Engine/Source/Runtime/WebBrowser/WebBrowser.Build.cs b/Engine/Source/Runtime/WebBrowser/WebBrowser.Build.cs
--- a/Engine/Source/Runtime/WebBrowser/WebBrowser.Build.cs
+++ b/Engine/Source/Runtime/WebBrowser/WebBrowser.Build.cs
@@ -16,8 +16,7 @@ public class WebBrowser : ModuleRules
 				"RHI",
 				"InputCore",
 				"Serialization",
-				"HTTP",
-				"Engine"
+				"HTTP"
 			}
 		);
 
@@ -61,25 +60,14 @@ public class WebBrowser : ModuleRules
 			CircularlyReferencedDependentModules.Add("WebBrowserTexture");
 		}
 		
-		if (Target.bCompileAgainstEngine && Target.Platform == UnrealTargetPlatform.Linux)
+		if (Target.bCompileAgainstEngine)
 		{
-			PrivateDependencyModuleNames.AddRange(
-				new string[]
-				{
-					"Engine"
-				}
-			);
-		}
+			if (Target.Platform == UnrealTargetPlatform.Win64)
+			{
+				PrivateDependencyModuleNames.Add("RenderCore");
+			}
 		
-		if (Target.bCompileAgainstEngine && Target.Platform == UnrealTargetPlatform.Win64)
-        {
-			PrivateDependencyModuleNames.AddRange(
-				new string[]
-				{
-					"Engine",
-					"RenderCore"
-				}
-			);
+			PrivateDependencyModuleNames.Add("Engine");
 		}
 
 		if (Target.Platform == UnrealTargetPlatform.Android)
