Adding support for the SAppleWebBrowser widget to the Slate widget reflector

#Synced-CL Unknown

diff --git a/Engine/Source/Runtime/Slate/Private/Framework/Application/SlateApplication.cpp b/Engine/Source/Runtime/Slate/Private/Framework/Application/SlateApplication.cpp
--- a/Engine/Source/Runtime/Slate/Private/Framework/Application/SlateApplication.cpp
+++ b/Engine/Source/Runtime/Slate/Private/Framework/Application/SlateApplication.cpp
@@ -4366,7 +4366,8 @@ bool FSlateApplication::TakeHDRScreenshot(const TSharedRef<SWidget>& Widget, TAr
 	return TakeHDRScreenshot(Widget, FIntRect(), OutColorData, OutSize);
 }
 
-void TakeScreenshotCommon(const TSharedRef<SWidget>& Widget, const FIntRect& InnerWidgetArea, FIntRect& ScreenshotRect, SWindow* WidgetWindow)
+void TakeScreenshotCommon(const TSharedRef<SWidget>& Widget, const FIntRect& InnerWidgetArea, FIntRect& ScreenshotRect, 
+						  SWindow* WidgetWindow)
 {
 	FWidgetPath WidgetPath;
 	FSlateApplication::Get().GeneratePathToWidgetChecked(Widget, WidgetPath);
@@ -4382,6 +4383,8 @@ void TakeScreenshotCommon(const TSharedRef<SWidget>& Widget, const FIntRect& Inn
 	ScreenshotRect.Min.Y += ( Position.Y - WindowPosition.Y );
 	ScreenshotRect.Max.X += ( Position.X - WindowPosition.X );
 	ScreenshotRect.Max.Y += ( Position.Y - WindowPosition.Y );
+	
+	FSlateApplication::Get().OnPreScreenshot().Broadcast(Widget);
 }
 
 bool FSlateApplication::TakeScreenshot(const TSharedRef<SWidget>& Widget, const FIntRect& InnerWidgetArea, TArray<FColor>& OutColorData, FIntVector& OutSize)
diff --git a/Engine/Source/Runtime/Slate/Public/Framework/Application/SlateApplication.h b/Engine/Source/Runtime/Slate/Public/Framework/Application/SlateApplication.h
--- a/Engine/Source/Runtime/Slate/Public/Framework/Application/SlateApplication.h
+++ b/Engine/Source/Runtime/Slate/Public/Framework/Application/SlateApplication.h
@@ -587,6 +587,10 @@ public:
 	/** Delegate called just before possible focus change */
 	DECLARE_MULTICAST_DELEGATE_FiveParams(FOnFocusChanging, const FFocusEvent&, const FWeakWidgetPath&, const TSharedPtr<SWidget>&, const FWidgetPath&, const TSharedPtr<SWidget>&);
 	FOnFocusChanging& OnFocusChanging() { return FocusChangingDelegate; }
+	
+	/** Delegate called when we are about to take a screenshot of a SWidget. */
+	DECLARE_MULTICAST_DELEGATE_OneParam(FOnPreScreenshot, TSharedRef<SWidget>)
+	FOnPreScreenshot& OnPreScreenshot() { return PreScreenshotDelegate; }
 
 	/** 
 	 * Removes references to FViewportRHI's.  
@@ -2167,6 +2171,9 @@ private:
 
 	/** Delegate for when focus might be about to change */
 	FOnFocusChanging FocusChangingDelegate;
+	
+	/** Delegate for when we are about to take a screenshot of a Slate widget */
+	FOnPreScreenshot PreScreenshotDelegate;
 
 	/** Critical section to avoid multiple threads calling Slate Tick when we're synchronizing between the Slate Loading Thread and the Game Thread. */
 	FCriticalSection SlateTickCriticalSection;
diff --git a/Engine/Source/Runtime/WebBrowser/Private/Apple/ApplePlatformWebBrowser.cpp b/Engine/Source/Runtime/WebBrowser/Private/Apple/ApplePlatformWebBrowser.cpp
--- a/Engine/Source/Runtime/WebBrowser/Private/Apple/ApplePlatformWebBrowser.cpp
+++ b/Engine/Source/Runtime/WebBrowser/Private/Apple/ApplePlatformWebBrowser.cpp
@@ -51,10 +51,12 @@ class SAppleWebBrowserWidget : public SViewport
 	SLATE_BEGIN_ARGS(SAppleWebBrowserWidget)
 		: _InitialURL("about:blank")
 		, _UseTransparency(false)
+		, _AllowScreenshots(true)
 	{ }
 
 	SLATE_ARGUMENT(FString, InitialURL);
 	SLATE_ARGUMENT(bool, UseTransparency);
+	SLATE_ARGUMENT(bool, AllowScreenshots);
 	SLATE_ARGUMENT(TSharedPtr<FWebBrowserWindow>, WebBrowserWindow);
 	SLATE_ARGUMENT(FString, UserAgentApplication);
 
@@ -100,6 +102,9 @@ class SAppleWebBrowserWidget : public SViewport
 		WebBrowserTexture = nullptr;
 		WebBrowserMaterial = nullptr;
 		WebBrowserBrush = nullptr;
+		
+		[WebViewWrapper setUsingIndirectRendering: false];
+		[WebViewWrapper setAllowingScreenshots: Args._AllowScreenshots];
 
 		// create external texture
 		WebBrowserTexture = NewObject<UWebBrowserTexture>((UObject*)GetTransientPackage(), NAME_None, RF_Transient | RF_Public);
@@ -131,9 +136,59 @@ class SAppleWebBrowserWidget : public SViewport
 				}
 			}
 		}
+		
+		check(![WebViewWrapper isAllowingScreenshots] || WebBrowserBrush.IsValid());
 #endif
+		
+		TSharedRef<SAppleWebBrowserWidget> SharedThis = 
+			StaticCastSharedPtr<SAppleWebBrowserWidget>(AsShared().ToSharedPtr()).ToSharedRef();
+		
+		FSlateApplication::Get().OnPreScreenshot().AddSP(SharedThis, &SAppleWebBrowserWidget::OnTakingScreenshot);
+		
 		LoadURL(Args._InitialURL);
 	}
+	
+	bool UsingIndirectRendering() const
+	{
+		return Is3DBrowser || (WebViewWrapper != nil && [WebViewWrapper isUsingIndirectRendering]);
+	}
+	
+	void OnTakingScreenshot(TSharedRef<SWidget> Widget)
+	{
+		if(![WebViewWrapper isAllowingScreenshots] || (!CurrentWidgetPath.IsValid() && 
+		   !FSlateApplication::Get().GeneratePathToWidgetUnchecked(AsShared(), CurrentWidgetPath)))
+		{
+			return;
+		}
+		
+		bool bIsVisible = !WebBrowserWindowPtr.IsValid() || WebBrowserWindowPtr.Pin()->IsVisible();
+	
+		if(CurrentWidgetPath.ContainsWidget(&Widget.Get()) && WebViewWrapper != nil &&
+		   !Is3DBrowser && bIsVisible)
+		{
+			// we are a descendent from this widget, therefore it is taking 
+			// a screenshot of us and we must switch to a Slate-visible view
+			[WebViewWrapper setUsingIndirectRendering: true];
+		}
+		
+		bTakingScreenshot = true;
+	}
+	
+	void ScreenshotFinished()
+	{
+		if(WebViewWrapper != nil && [WebViewWrapper isUsingIndirectRendering])
+		{
+			check([WebViewWrapper isAllowingScreenshots]);
+			
+			bForceRecomputeWindowOnPaint = true;
+			bTakingScreenshot = false;
+			
+			if(!Is3DBrowser)
+			{
+				[WebViewWrapper setUsingIndirectRendering: false];
+			}
+		}
+	}
 
 	void UpdateFrameWithViewport(FIntPoint WindowSize, FIntPoint WindowPos) {
 		CGRect NewFrame;
@@ -185,7 +240,9 @@ class SAppleWebBrowserWidget : public SViewport
 		}
 #endif
 
-		if(!Is3DBrowser && (bForceRecomputeWindowOnPaint || NewFrame.size.width != WindowFrame.size.width || NewFrame.size.height != WindowFrame.size.height)) 
+		if(!UsingIndirectRendering() && 
+		   (bForceRecomputeWindowOnPaint || NewFrame.size.width != WindowFrame.size.width || 
+			NewFrame.size.height != WindowFrame.size.height)) 
 		{
 			// WindowFrame is only used when we are a 2D web browser
 			WindowFrame = NewFrame;
@@ -290,8 +347,19 @@ class SAppleWebBrowserWidget : public SViewport
 									FIntPoint(static_cast<int32_t>(Position.X), static_cast<int32_t>(Position.Y)));
 			
 #if !PLATFORM_TVOS
-			if (Is3DBrowser)
+			if ([WebViewWrapper isAllowingScreenshots] || UsingIndirectRendering())
 			{
+				if(!bFirstScreenshotTaken && !Is3DBrowser)
+				{
+					// If we are allowing screenshots in 2D mode, we want to take one
+					// screenshot at first to make sure that everything gets properly
+					// initialized and the indirect texture set up properly.
+					[WebViewWrapper setUsingIndirectRendering: true];
+					bTakingScreenshot = true;
+					bFirstScreenshotTaken = true;
+				}
+				
+				
 				if (WebBrowserTexture.IsValid())
 				{					
 					TSharedPtr<FWebBrowserTextureSample, ESPMode::ThreadSafe> WebBrowserTextureSample;
@@ -304,29 +372,34 @@ class SAppleWebBrowserWidget : public SViewport
 				{
 					TSharedPtr<FWebBrowserWindow> WebBrowserWindow = WebBrowserWindowPtr.Pin();
 					bool TextureWritten = false;
-										
-					// GetViewportSize not 100% accurate in 3D mode, we use it as a 
-					// reasonable default if we need to only
-					FIntPoint ViewportSize = WebBrowserWindow->GetViewportSize();
 					
-					if(TSharedPtr<SWindow> ParentWindow = WebBrowserWindow->GetParentWindow())
+					FIntPoint ViewportSize{WindowFrame.size.width, WindowFrame.size.height};
+					
+					if(Is3DBrowser)
 					{
-						FVector2f WindowSize = ParentWindow->GetSizeInScreen();
-						ViewportSize.X = static_cast<int32>(WindowSize.X);
-						ViewportSize.Y = static_cast<int32>(WindowSize.Y);
+						// GetViewportSize not 100% accurate in 3D mode, we use it as a 
+						// reasonable default if we need to only
+						ViewportSize = WebBrowserWindow->GetViewportSize();
+						
+						if(TSharedPtr<SWindow> ParentWindow = WebBrowserWindow->GetParentWindow())
+						{
+							FVector2f WindowSize = ParentWindow->GetSizeInScreen();
+							ViewportSize.X = static_cast<int32>(WindowSize.X);
+							ViewportSize.Y = static_cast<int32>(WindowSize.Y);
 #if !PLATFORM_MAC
-						// needs to be on main thread for getScreenScale
-						dispatch_async(dispatch_get_main_queue(), ^ 
-					    {
-							FIntPoint ScaledViewportSize = ViewportSize;
-							CGFloat ScreenScale = [WebViewWrapper getScreenScale];
-							ScaledViewportSize.X *= ScreenScale;
-							ScaledViewportSize.Y *= ScreenScale;
-							
-							WriteWebBrowserTexture(WebViewWrapper, WebBrowserTexture->GetExternalTextureGuid(), ScaledViewportSize);
-						});
-						TextureWritten = true;
+							// needs to be on main thread for getScreenScale
+							dispatch_async(dispatch_get_main_queue(), ^ 
+							{
+								FIntPoint ScaledViewportSize = ViewportSize;
+								CGFloat ScreenScale = [WebViewWrapper getScreenScale];
+								ScaledViewportSize.X *= ScreenScale;
+								ScaledViewportSize.Y *= ScreenScale;
+								
+								WriteWebBrowserTexture(WebViewWrapper, WebBrowserTexture->GetExternalTextureGuid(), ScaledViewportSize);
+							});
+							TextureWritten = true;
 #endif
+						}
 					}
 
 					if(!TextureWritten) 
@@ -337,7 +410,6 @@ class SAppleWebBrowserWidget : public SViewport
 			}
 #endif
 		}
-
 	}
 
 	int32 OnPaint(const FPaintArgs& Args, const FGeometry& AllottedGeometry, const FSlateRect& MyCullingRect, FSlateWindowElementList& OutDrawElements, int32 LayerId, const FWidgetStyle& InWidgetStyle, bool bParentEnabled) const
@@ -345,7 +417,7 @@ class SAppleWebBrowserWidget : public SViewport
 #if !PLATFORM_TVOS
 		bool bIsVisible = !WebBrowserWindowPtr.IsValid() || WebBrowserWindowPtr.Pin()->IsVisible();
 		
-		if (bIsVisible && Is3DBrowser && WebBrowserBrush.IsValid())
+		if (bIsVisible && UsingIndirectRendering() && WebBrowserBrush.IsValid())
 		{
 			FSlateDrawElement::MakeBox(OutDrawElements, LayerId, AllottedGeometry.ToPaintGeometry(), WebBrowserBrush.Get(), ESlateDrawEffect::None);
 		}
@@ -667,6 +739,8 @@ private:
 	bool bAcquiredInitialFocus;
 	bool bFocused;
 	bool bForceRecomputeWindowOnPaint = false;
+	bool bTakingScreenshot = false;
+	bool bFirstScreenshotTaken = false;
 
 #if !PLATFORM_TVOS
 	/** The external texture to render the webbrowser output. */
@@ -688,6 +762,8 @@ private:
 	FDelegateHandle OnSafeFrameChangedEventHandle;
 #endif
 	
+	FWidgetPath CurrentWidgetPath;
+	
 	CGRect WindowFrame;
 };
 
@@ -703,6 +779,7 @@ private:
 #endif
 @synthesize WebViewContainer;
 @synthesize UserContentController;
+@synthesize LastWebViewScreenshot;
 #endif
 @synthesize NextURL;
 @synthesize NextContent;
@@ -854,6 +931,12 @@ supportsMetal : (bool)InSupportsMetal supportsMetalMRT : (bool)InSupportsMetalMR
 		WebBrowserWidget = nil;
 	}
 
+	if(LastWebViewScreenshot != nil)
+	{
+		[LastWebViewScreenshot release];
+		LastWebViewScreenshot = nil;
+	}
+	
 #if !PLATFORM_MAC
 	if (CloseButton != nil)
 	{
@@ -969,7 +1052,7 @@ supportsMetal : (bool)InSupportsMetal supportsMetalMRT : (bool)InSupportsMetalMR
 				self.WebView.frame = WebViewContainer.bounds;
 			}
 #else
-			if(!Is3DBrowser)
+			if(![self isUsingIndirectRendering])
 			{
 				self.WebViewContainer.frame = self.DesiredFrame;
 				self.WebView.frame = WebViewContainer.bounds;
@@ -1080,6 +1163,16 @@ supportsMetal : (bool)InSupportsMetal supportsMetalMRT : (bool)InSupportsMetalMR
 	});
 }
 
+-(bool)isUsingIndirectRendering;
+{
+	return Is3DBrowser || bUsingIndirectRendering;
+}
+
+-(bool)isAllowingScreenshots;
+{
+	return bAllowScreenshots;
+}
+
 -(void)set3D:(bool)InIs3DBrowser;
 {
 	dispatch_async(dispatch_get_main_queue(), ^
@@ -1093,6 +1186,17 @@ supportsMetal : (bool)InSupportsMetal supportsMetalMRT : (bool)InSupportsMetalMR
 	});
 }
 
+-(void)setUsingIndirectRendering:(bool)UsingIndirectRendering;
+{
+	bUsingIndirectRendering = UsingIndirectRendering;
+	[self setDefaultVisibility];
+}
+
+-(void)setAllowingScreenshots:(bool)AllowScreenshots;
+{
+	bAllowScreenshots = AllowScreenshots;
+}
+
 -(void)setDefaultVisibility;
 {
 #if !PLATFORM_TVOS
@@ -1211,8 +1315,6 @@ supportsMetal : (bool)InSupportsMetal supportsMetalMRT : (bool)InSupportsMetalMR
 	@synchronized(self) // Briefly block render thread
 	{
 		id<MTLTexture> ptrToMetalTexture = (id<MTLTexture>)ptr;
-		NSUInteger width = [ptrToMetalTexture width];
-		NSUInteger height = [ptrToMetalTexture height];
 
 		[self updateWebViewMetalTexture : ptrToMetalTexture];
 	}
@@ -1243,11 +1345,53 @@ supportsMetal : (bool)InSupportsMetal supportsMetalMRT : (bool)InSupportsMetalMR
 }
 #endif
 
+#if PLATFORM_MAC
+-(void)replaceMetalTexture:(id<MTLTexture>)Texture fromImage:(NSImage*)Image;
+#else
+-(void)replaceMetalTexture:(id<MTLTexture>)Texture fromImage:(UIImage*)Image;
+#endif
+{
+#if PLATFORM_MAC
+	if(Image == nil)
+#else
+	if(Image == nil || Image.CGImage == nil)
+#endif
+	{
+		return;
+	}
+	
+	NSUInteger TextureWidth = [Texture width];
+	NSUInteger TextureHeight = [Texture height];
+	
+	CGColorSpaceRef ColorSpace = CGColorSpaceCreateDeviceRGB();
+	CGContextRef WebViewContext = CGBitmapContextCreate(NULL, TextureWidth, TextureHeight, 8, 4 * TextureWidth, ColorSpace, (CGBitmapInfo)kCGImageAlphaPremultipliedLast);
+	
+#if PLATFORM_MAC
+	NSGraphicsContext *NewGfxContext = [NSGraphicsContext graphicsContextWithCGContext : WebViewContext
+														  flipped : NO];
+	NSGraphicsContext *OldGfxContext = NSGraphicsContext.currentContext;
+	NSGraphicsContext.currentContext = NewGfxContext;
+	
+	[Image drawInRect: WebView.frame];
+ 
+	NSGraphicsContext.currentContext = OldGfxContext;
+#else
+	CGContextDrawImage(WebViewContext, CGRectMake(0, 0, TextureWidth, TextureHeight), Image.CGImage);
+#endif
+	[Texture replaceRegion : MTLRegionMake2D(0, 0, TextureWidth, TextureHeight)
+		 mipmapLevel : 0
+		 withBytes : CGBitmapContextGetData(WebViewContext)
+		 bytesPerRow : 4 * TextureWidth];
+ 
+	CGColorSpaceRelease(ColorSpace);
+	CGContextRelease(WebViewContext);
+}
+
 -(void)updateWebViewMetalTexture:(id<MTLTexture>)texture
 {
 #if !PLATFORM_TVOS
 	
-	dispatch_async(dispatch_get_main_queue(), ^ 
+	dispatch_sync(dispatch_get_main_queue(), ^ 
 	{
 		@autoreleasepool 
 		{
@@ -1278,41 +1422,38 @@ supportsMetal : (bool)InSupportsMetal supportsMetalMRT : (bool)InSupportsMetalMR
 #else
 			 ^ void(UIImage *Image, NSError *Error)
 #endif
-			{				
+			{
 				if(Error != nil) 
 				{
 					UE_LOG(LogAppleWebBrowser, Error, TEXT("takeSnapshotWithConfiguration errored with NSError %d!"), Error.code);
-					
-					[SnapshotConfig release];
-					return;
 				}
-				
-				CGColorSpaceRef ColorSpace = CGColorSpaceCreateDeviceRGB();
-				CGContextRef WebViewContext = CGBitmapContextCreate(NULL, TextureWidth, TextureHeight, 8, 4 * TextureWidth, ColorSpace, (CGBitmapInfo)kCGImageAlphaPremultipliedLast);
-				
-#if PLATFORM_MAC
-				NSGraphicsContext *NewGfxContext = [NSGraphicsContext graphicsContextWithCGContext : WebViewContext
-																	  flipped : NO];
-				NSGraphicsContext *OldGfxContext = NSGraphicsContext.currentContext;
-				NSGraphicsContext.currentContext = NewGfxContext;
-				
-				[Image drawInRect: WebView.frame];
-			 
-				NSGraphicsContext.currentContext = OldGfxContext;
-#else
-				CGContextDrawImage(WebViewContext, CGRectMake(0, 0, TextureWidth, TextureHeight), Image.CGImage);
-#endif
-				[texture replaceRegion : MTLRegionMake2D(0, 0, TextureWidth, TextureHeight)
-					 mipmapLevel : 0
-					 withBytes : CGBitmapContextGetData(WebViewContext)
-					 bytesPerRow : 4 * TextureWidth];
+				else if(Is3DBrowser)
+				{
+					[self replaceMetalTexture: texture fromImage: Image];
+				}
+				else if(bAllowScreenshots)
+				{
+					if(LastWebViewScreenshot != nil)
+					{
+						[LastWebViewScreenshot release];
+					}
+					LastWebViewScreenshot = Image;
+					[LastWebViewScreenshot retain];
+				}
 				
 				Image = nil;
-			 
-				CGColorSpaceRelease(ColorSpace);
-				CGContextRelease(WebViewContext);
 				[SnapshotConfig release];
 			}];
+			
+			if(!Is3DBrowser && [self isUsingIndirectRendering])
+			{
+				// always take screenshots to get around async weirdness but never 
+				// do an expensive copy to the Slate texture unless we need to because
+				// we are actively taking a screenshot this frame.
+				
+				check(bAllowScreenshots);
+				[self replaceMetalTexture: texture fromImage: LastWebViewScreenshot];
+			}
 		}
 	});
 #endif
@@ -1821,6 +1962,8 @@ void FWebBrowserWindow::CheckTickActivity()
 	}
 
 	bTickedLastFrame = false;
+	
+	BrowserWidget->ScreenshotFinished();
 }
 
 void FWebBrowserWindow::SetTickLastFrame()
diff --git a/Engine/Source/Runtime/WebBrowser/Private/Apple/ApplePlatformWebBrowser.h b/Engine/Source/Runtime/WebBrowser/Private/Apple/ApplePlatformWebBrowser.h
--- a/Engine/Source/Runtime/WebBrowser/Private/Apple/ApplePlatformWebBrowser.h
+++ b/Engine/Source/Runtime/WebBrowser/Private/Apple/ApplePlatformWebBrowser.h
@@ -51,15 +51,19 @@ class SWebBrowserView;
 	bool bVideoTextureValid;
 	bool bSupportsMetal;
 	bool bSupportsMetalMRT;
+	bool bUsingIndirectRendering;
+	bool bAllowScreenshots;
 }
 #if !PLATFORM_TVOS
 @property(strong) WKWebView* WebView;
 #if PLATFORM_MAC
 @property(strong) NSView* WebViewContainer;
 @property(strong) NSView* GameWindowView;
+@property(strong) NSImage* LastWebViewScreenshot;
 #else
 @property(strong) WebViewCloseButton* CloseButton;
 @property(strong) UIView* WebViewContainer;
+@property(strong) UIImage* LastWebViewScreenshot;
 #endif
 @property(strong) WKUserContentController* UserContentController;
 #endif
@@ -77,6 +81,10 @@ supportsMetal : (bool)InSupportsMetal supportsMetalMRT : (bool)InSupportsMetalMR
 -(void)loadstring:(NSString*)InString dummyurl:(NSURL*)InURL;
 -(void)loadurl:(NSURL*)InURL;
 -(void)executejavascript:(NSString*)InJavaScript;
+-(bool)isUsingIndirectRendering;
+-(void)setUsingIndirectRendering:(bool)UsingIndirectRendering;
+-(bool)isAllowingScreenshots;
+-(void)setAllowingScreenshots:(bool)AllowScreenshots;
 -(void)set3D:(bool)InIs3DBrowser;
 -(void)setDefaultVisibility;
 -(void)setVisibility:(bool)InIsVisible;
@@ -91,6 +99,11 @@ supportsMetal : (bool)InSupportsMetal supportsMetalMRT : (bool)InSupportsMetalMR
 -(void)SetVideoTextureValid:(bool)Condition;
 -(bool)IsVideoTextureValid;
 -(bool)UpdateVideoFrame:(void*)ptr;
+#if PLATFORM_MAC
+-(void)replaceMetalTexture:(id<MTLTexture>)Texture fromImage:(NSImage*)Image;
+#else
+-(void)replaceMetalTexture:(id<MTLTexture>)Texture fromImage:(UIImage*)Image;
+#endif
 -(void)updateWebViewMetalTexture : (id<MTLTexture>)texture;
 -(void)leaveDeadWindow;
 -(bool)inKeyWindow;
