WebBrowser: Fix build with older CEF version

#rnx

diff --git a/Engine/Source/Runtime/WebBrowser/Private/WebBrowserSingleton.cpp b/Engine/Source/Runtime/WebBrowser/Private/WebBrowserSingleton.cpp
--- a/Engine/Source/Runtime/WebBrowser/Private/WebBrowserSingleton.cpp
+++ b/Engine/Source/Runtime/WebBrowser/Private/WebBrowserSingleton.cpp
@@ -424,7 +424,7 @@ FWebBrowserSingleton::FWebBrowserSingleton(const FWebBrowserInitSettings& WebBro
 				break;
 			}
 
-#if CEF3_DEFAULT_CACHE
+#if CEF_VERSION_MAJOR >= 128
 			const int CefExitCode = CefGetExitCode();
 			// There's a race condition that can occur when two UE instances are started close together:
 			// We check the presence of a lockfile inside the CEF cache dir to detect if that cache is already in use by another UE instance,
@@ -467,7 +467,7 @@ FWebBrowserSingleton::FWebBrowserSingleton(const FWebBrowserInitSettings& WebBro
 #endif // PLATFORM_WINDOWS
 			}
 			UE_LOG(LogWebBrowser, Error, TEXT("CEF initialization for cache dir %s failed with exit code %d!"), *CachePath, CefExitCode);
-#endif // CEF3_DEFAULT_CACHE
+#endif // CEF_VERSION_MAJOR >= 128
 			break;
 		}
 		// this will upload a report to our online crash tool so we can find out why the initialization failed
