Including moved WebBrowser assets in build when using WebBrowserWidget

#Synced-CL Unknown

diff --git a/Engine/Plugins/Runtime/WebBrowserWidget/Config/DefaultWebBrowserWidget.ini b/Engine/Plugins/Runtime/WebBrowserWidget/Config/DefaultWebBrowserWidget.ini
new file mode 100644
--- /dev/null
+++ b/Engine/Plugins/Runtime/WebBrowserWidget/Config/DefaultWebBrowserWidget.ini
@@ -0,0 +1,2 @@
+[/Script/UnrealEd.ProjectPackagingSettings]
++DirectoriesToAlwaysCook=(Path="/Engine/WebBrowser/")
diff --git a/Engine/Source/Runtime/WebBrowser/Private/WebBrowserSingleton.h b/Engine/Source/Runtime/WebBrowser/Private/WebBrowserSingleton.h
--- a/Engine/Source/Runtime/WebBrowser/Private/WebBrowserSingleton.h
+++ b/Engine/Source/Runtime/WebBrowser/Private/WebBrowserSingleton.h
@@ -125,7 +125,7 @@ public:
 	/** Set a reference to UWebBrowser's default material*/
 	virtual void SetDefaultMaterial(UMaterialInterface* InDefaultMaterial) override
 	{
-#if WITH_EDITOR
+#if WITH_ENGINE
 		DefaultMaterial = InDefaultMaterial;
 #endif
 	}
@@ -133,7 +133,7 @@ public:
 	/** Set a reference to UWebBrowser's translucent material*/
 	virtual void SetDefaultTranslucentMaterial(UMaterialInterface* InDefaultMaterial) override
 	{
-#if WITH_EDITOR
+#if WITH_ENGINE
 		DefaultTranslucentMaterial = InDefaultMaterial;
 #endif
 	}
@@ -141,7 +141,7 @@ public:
 	/** Get a reference to UWebBrowser's default material*/
 	virtual UMaterialInterface* GetDefaultMaterial() override
 	{
-#if WITH_EDITOR
+#if WITH_ENGINE
 		return DefaultMaterial.Get();
 #else
 		return nullptr;
@@ -151,7 +151,7 @@ public:
 	/** Get a reference to UWebBrowser's translucent material*/
 	virtual UMaterialInterface* GetDefaultTranslucentMaterial() override
 	{
-#if WITH_EDITOR
+#if WITH_ENGINE
 		return DefaultTranslucentMaterial.Get();
 #else
 		return nullptr;
