fixing Apple 3D web browser leak on iOS

#Synced-CL 49122214

diff --git a/Engine/Source/Runtime/WebBrowser/Private/Apple/ApplePlatformWebBrowser.cpp b/Engine/Source/Runtime/WebBrowser/Private/Apple/ApplePlatformWebBrowser.cpp
--- a/Engine/Source/Runtime/WebBrowser/Private/Apple/ApplePlatformWebBrowser.cpp
+++ b/Engine/Source/Runtime/WebBrowser/Private/Apple/ApplePlatformWebBrowser.cpp
@@ -1400,7 +1400,10 @@ private:
 			UIView* View = [IOSAppDelegate GetDelegate].IOSView;
 			CGFloat ContentScaleFactor = View.contentScaleFactor;
 			
+			// snapshotWidth needs to be set to the pre-dpi scaling size of the WebView
+			// to avoid an internal WKWebView memory leak
 			WebView.frame = CGRectMake(0, 0, ContentScaleFactor * Width, ContentScaleFactor * Height);
+			SnapshotConfig.snapshotWidth = [NSNumber numberWithFloat: WebView.bounds.size.width / ContentScaleFactor];
 #endif
 			SnapshotConfig.rect = CGRectMake(0, 0, WebView.bounds.size.width, WebView.bounds.size.height);
 			
