Added defensive checks against callbacks in ApplePlatformWebBrowser from going ahead upon module shutdown

#Synced-CL 49225950

diff --git a/Engine/Source/Runtime/WebBrowser/Private/Apple/ApplePlatformWebBrowser.cpp b/Engine/Source/Runtime/WebBrowser/Private/Apple/ApplePlatformWebBrowser.cpp
--- a/Engine/Source/Runtime/WebBrowser/Private/Apple/ApplePlatformWebBrowser.cpp
+++ b/Engine/Source/Runtime/WebBrowser/Private/Apple/ApplePlatformWebBrowser.cpp
@@ -136,6 +136,7 @@ class SAppleWebBrowserWidget : public SViewport
 			StaticCastSharedPtr<SAppleWebBrowserWidget>(AsShared().ToSharedPtr()).ToSharedRef();
 		
 		FSlateApplication::Get().OnPreScreenshot().AddSP(SharedThis, &SAppleWebBrowserWidget::OnTakingScreenshot);
+		FCoreDelegates::OnEnginePreExit.AddSP(SharedThis, &SAppleWebBrowserWidget::OnEnginePreExit);
 		
 		LoadURL(Args._InitialURL);
 	}
@@ -166,6 +167,14 @@ class SAppleWebBrowserWidget : public SViewport
 		bTakingScreenshot = true;
 	}
 	
+	void OnEnginePreExit()
+	{
+		if(WebViewWrapper != nil)
+		{
+			[WebViewWrapper disableCallbacksOnShutdown];
+		}
+	}
+	
 	void ScreenshotFinished()
 	{
 		if(WebViewWrapper != nil && [WebViewWrapper isUsingIndirectRendering])
@@ -788,6 +797,7 @@ private:
 	bNeedsAddToView = true;
 	Is3DBrowser = false;
 	bVideoTextureValid = false;
+	bShuttingDown = false;
 	
 #if !UE_BUILD_SHIPPING
 	BOOL Inspectable = CVarWebViewEnableInspectable.GetValueOnAnyThread()? YES : NO;
@@ -1230,6 +1240,11 @@ private:
 #endif
 }
 
+-(void)disableCallbacksOnShutdown;
+{
+	bShuttingDown = true;
+}
+
 -(void)reload;
 {
 #if !PLATFORM_TVOS
@@ -1487,57 +1502,98 @@ private:
 
 - (void)webView:(WKWebView*)InWebView decidePolicyForNavigationAction : (WKNavigationAction*)InNavigationAction decisionHandler : (void(^)(WKNavigationActionPolicy))InDecisionHandler
 {
-	NSURLRequest *request = InNavigationAction.request;
-	FString UrlStr([[request URL] absoluteString]);
-
-	if (InNavigationAction.targetFrame == nil && !UrlStr.IsEmpty() && FPlatformProcess::CanLaunchURL(*UrlStr))
+	if(IWebBrowserModule::IsAvailable() && !bShuttingDown)
 	{
-		if (WebBrowserWidget->HandleOnBeforePopup(UrlStr, TEXT("_blank")))
+		NSURLRequest *request = InNavigationAction.request;
+		FString UrlStr([[request URL] absoluteString]);
+
+		if (InNavigationAction.targetFrame == nil && !UrlStr.IsEmpty() && FPlatformProcess::CanLaunchURL(*UrlStr))
 		{
-			// Launched the URL in external browser, don't open the link here too
-			InDecisionHandler(WKNavigationActionPolicyCancel);
-			return;
+			if (WebBrowserWidget->HandleOnBeforePopup(UrlStr, TEXT("_blank")))
+			{
+				// Launched the URL in external browser, don't open the link here too
+				InDecisionHandler(WKNavigationActionPolicyCancel);
+				return;
+			}
 		}
+		
+		WebBrowserWidget->HandleShouldOverrideUrlLoading(UrlStr);
+		InDecisionHandler(WKNavigationActionPolicyAllow);
+	}
+	else
+	{
+		UE_LOG(LogAppleWebBrowser, Warning, 
+			   TEXT("ApplePlatformWebBrowser recieved decidePolicyForNavigationAction callback during or after module shutdown, preventing navigation"));
+		InDecisionHandler(WKNavigationActionPolicyCancel);
 	}
-	
-	WebBrowserWidget->HandleShouldOverrideUrlLoading(UrlStr);
-	InDecisionHandler(WKNavigationActionPolicyAllow);
 }
 
 -(void)webView:(WKWebView *)InWebView didCommitNavigation : (WKNavigation *)InNavigation
 {
-	NSString* CurrentUrl = [self.WebView URL].absoluteString;
-	NSString* Title = [self.WebView title];
-	
-//	NSLog(@"didCommitNavigation: %@", CurrentUrl);
-	WebBrowserWidget->HandleReceivedTitle(Title);
-	WebBrowserWidget->HandlePageLoad(CurrentUrl, true);
+	if(IWebBrowserModule::IsAvailable() && !bShuttingDown)
+	{
+		NSString* CurrentUrl = [self.WebView URL].absoluteString;
+		NSString* Title = [self.WebView title];
+		
+	//	NSLog(@"didCommitNavigation: %@", CurrentUrl);
+		WebBrowserWidget->HandleReceivedTitle(Title);
+		WebBrowserWidget->HandlePageLoad(CurrentUrl, true);
+	}
+	else
+	{
+		UE_LOG(LogAppleWebBrowser, Warning, 
+			   TEXT("ApplePlatformWebBrowser recieved didCommitNavigation callback during or after module shutdown, ignoring"));
+	}
 }
 
 -(void)webView:(WKWebView *)InWebView didFinishNavigation : (WKNavigation *)InNavigation
 {
-	NSString* CurrentUrl = [self.WebView URL].absoluteString;
-	NSString* Title = [self.WebView title];
-	// NSLog(@"didFinishNavigation: %@", CurrentUrl);
-	WebBrowserWidget->HandleReceivedTitle(Title);
-	WebBrowserWidget->HandlePageLoad(CurrentUrl, false);
+	if(IWebBrowserModule::IsAvailable() && !bShuttingDown)
+	{
+		NSString* CurrentUrl = [self.WebView URL].absoluteString;
+		NSString* Title = [self.WebView title];
+		// NSLog(@"didFinishNavigation: %@", CurrentUrl);
+		WebBrowserWidget->HandleReceivedTitle(Title);
+		WebBrowserWidget->HandlePageLoad(CurrentUrl, false);
+	}
+	else
+	{
+		UE_LOG(LogAppleWebBrowser, Warning, 
+			   TEXT("ApplePlatformWebBrowser recieved didFinishNavigation callback during or after module shutdown, ignoring"));
+	}
 }
 -(void)webView:(WKWebView *)InWebView didFailNavigation : (WKNavigation *)InNavigation withError : (NSError*)InError
 {
-	if (InError.domain == NSURLErrorDomain && InError.code == NSURLErrorCancelled)
+	if(IWebBrowserModule::IsAvailable() && !bShuttingDown)
 	{
-		//ignore this one, interrupted load
-		return;
+		if (InError.domain == NSURLErrorDomain && InError.code == NSURLErrorCancelled)
+		{
+			//ignore this one, interrupted load
+			return;
+		}
+		NSString* CurrentUrl = [InError.userInfo objectForKey : @"NSErrorFailingURLStringKey"];
+	//	NSLog(@"didFailNavigation: %@, error %@", CurrentUrl, InError);
+		WebBrowserWidget->HandleReceivedError(InError.code, CurrentUrl);
+	}
+	else
+	{
+		UE_LOG(LogAppleWebBrowser, Warning, 
+			   TEXT("ApplePlatformWebBrowser recieved didFailNavigation callback during or after module shutdown, ignoring"));
 	}
-	NSString* CurrentUrl = [InError.userInfo objectForKey : @"NSErrorFailingURLStringKey"];
-//	NSLog(@"didFailNavigation: %@, error %@", CurrentUrl, InError);
-	WebBrowserWidget->HandleReceivedError(InError.code, CurrentUrl);
 }
 -(void)webView:(WKWebView *)InWebView didFailProvisionalNavigation : (WKNavigation *)InNavigation withError : (NSError*)InError
 {
-	NSString* CurrentUrl = [InError.userInfo objectForKey : @"NSErrorFailingURLStringKey"];
-	// NSLog(@"didFailProvisionalNavigation: %@, error %@", CurrentUrl, InError);
-	WebBrowserWidget->HandleReceivedError(InError.code, CurrentUrl);
+	if(IWebBrowserModule::IsAvailable() && !bShuttingDown)
+	{
+		NSString* CurrentUrl = [InError.userInfo objectForKey : @"NSErrorFailingURLStringKey"];
+		// NSLog(@"didFailProvisionalNavigation: %@, error %@", CurrentUrl, InError);
+		WebBrowserWidget->HandleReceivedError(InError.code, CurrentUrl);
+	}
+	else
+	{
+		UE_LOG(LogAppleWebBrowser, Warning, 
+			   TEXT("ApplePlatformWebBrowser recieved didFailProvisionalNavigation callback during or after module shutdown, ignoring"));
+	}
 }
 #endif
 @end
diff --git a/Engine/Source/Runtime/WebBrowser/Private/Apple/ApplePlatformWebBrowser.h b/Engine/Source/Runtime/WebBrowser/Private/Apple/ApplePlatformWebBrowser.h
--- a/Engine/Source/Runtime/WebBrowser/Private/Apple/ApplePlatformWebBrowser.h
+++ b/Engine/Source/Runtime/WebBrowser/Private/Apple/ApplePlatformWebBrowser.h
@@ -51,6 +51,10 @@ class SWebBrowserView;
 	bool bVideoTextureValid;
 	bool bUsingIndirectRendering;
 	bool bAllowScreenshots;
+	// needs to be atomic as it is set on the game thread
+	// but read from AppKit/UIKit's UI thread during WebKit
+	// callbacks
+	std::atomic<bool> bShuttingDown;
 }
 #if !PLATFORM_TVOS
 @property(strong) WKWebView* WebView;
@@ -86,6 +90,7 @@ enableFloatingCloseButton : (bool)bEnableFloatingCloseButton;
 -(void)set3D:(bool)InIs3DBrowser;
 -(void)setDefaultVisibility;
 -(void)setVisibility:(bool)InIsVisible;
+-(void)disableCallbacksOnShutdown;
 -(void)stopLoading;
 -(void)reload;
 -(void)goBack;
