Fix iOS 2D web browser by temporarily disabling slate screenshot support by default

#Synced-CL 49414982

diff --git a/Engine/Source/Runtime/WebBrowser/Private/Apple/ApplePlatformWebBrowser.cpp b/Engine/Source/Runtime/WebBrowser/Private/Apple/ApplePlatformWebBrowser.cpp
--- a/Engine/Source/Runtime/WebBrowser/Private/Apple/ApplePlatformWebBrowser.cpp
+++ b/Engine/Source/Runtime/WebBrowser/Private/Apple/ApplePlatformWebBrowser.cpp
@@ -51,7 +51,7 @@ class SAppleWebBrowserWidget : public SViewport
 	SLATE_BEGIN_ARGS(SAppleWebBrowserWidget)
 		: _InitialURL("about:blank")
 		, _UseTransparency(false)
-		, _AllowScreenshots(true)
+		, _AllowScreenshots(false)
 	{ }
 
 	SLATE_ARGUMENT(FString, InitialURL);
