[HomeScreen] SWebBrowser widget loading URL when size change stretch the old content
Hide the web view until the page is loaded and show the throbber

#Synced-CL 49539534

diff --git a/Engine/Source/Editor/MainFrame/Private/MainFrameModule.cpp b/Engine/Source/Editor/MainFrame/Private/MainFrameModule.cpp
--- a/Engine/Source/Editor/MainFrame/Private/MainFrameModule.cpp
+++ b/Engine/Source/Editor/MainFrame/Private/MainFrameModule.cpp
@@ -1243,9 +1243,28 @@ TSharedRef<SWidget> FMainFrameModule::CreateHomeScreenWidget(const TSharedPtr<SD
 		});
 	}
 
+	auto ShowThrobber = [this]()
+		{
+			if (WebBrowser.IsValid())
+			{
+				WebBrowser->ShowThrobber();
+			}
+		};
+
+	auto HideThrobber = [this]()
+		{
+			if (WebBrowser.IsValid())
+			{
+				WebBrowser->HideThrobber();
+			}
+		};
+
 	if (!WebBrowser.IsValid())
 	{
 		SAssignNew(WebBrowser, SWebBrowser, WebBrowserWindow)
+		.OnLoadStarted_Lambda(ShowThrobber)
+		.OnLoadError_Lambda(HideThrobber)
+		.OnLoadCompleted_Lambda(HideThrobber)
 		.ShowControls(false)
 		.SupportsTransparency(true)
 		.ShowAddressBar(false);
diff --git a/Engine/Source/Runtime/WebBrowser/Private/SWebBrowser.cpp b/Engine/Source/Runtime/WebBrowser/Private/SWebBrowser.cpp
--- a/Engine/Source/Runtime/WebBrowser/Private/SWebBrowser.cpp
+++ b/Engine/Source/Runtime/WebBrowser/Private/SWebBrowser.cpp
@@ -481,7 +481,7 @@ void SWebBrowser::OnUrlTextCommitted( const FText& NewText, ETextCommit::Type Co
 
 EVisibility SWebBrowser::GetViewportVisibility() const
 {
-	if (!bShowInitialThrobber || BrowserView->IsInitialized())
+	if (!bForceThrobber && (!bShowInitialThrobber || BrowserView->IsInitialized()))
 	{
 		return EVisibility::Visible;
 	}
@@ -490,7 +490,7 @@ EVisibility SWebBrowser::GetViewportVisibility() const
 
 EVisibility SWebBrowser::GetLoadingThrobberVisibility() const
 {
-	if (bShowInitialThrobber && !BrowserView->IsInitialized())
+	if (bForceThrobber || (bShowInitialThrobber && !BrowserView->IsInitialized()))
 	{
 		return EVisibility::Visible;
 	}
@@ -598,6 +598,16 @@ void SWebBrowser::SetParentWindow(TSharedPtr<SWindow> Window)
 	}
 }
 
+void SWebBrowser::ShowThrobber()
+{
+	bForceThrobber = true;
+}
+
+void SWebBrowser::HideThrobber()
+{
+	bForceThrobber = false;
+}
+
 FReply SWebBrowser::OnFocusReceived(const FGeometry& MyGeometry, const FFocusEvent& InFocusEvent)
 {
 	if (BrowserView.IsValid() && InFocusEvent.GetCause() == EFocusCause::SetDirectly)
diff --git a/Engine/Source/Runtime/WebBrowser/Public/SWebBrowser.h b/Engine/Source/Runtime/WebBrowser/Public/SWebBrowser.h
--- a/Engine/Source/Runtime/WebBrowser/Public/SWebBrowser.h
+++ b/Engine/Source/Runtime/WebBrowser/Public/SWebBrowser.h
@@ -248,6 +248,12 @@ public:
 	/** Set parent SWindow for this browser. */
 	WEBBROWSER_API void SetParentWindow(TSharedPtr<SWindow> Window);
 
+	/** Show the throbber and hide the web browser content. */
+	WEBBROWSER_API void ShowThrobber();
+
+	/** Hide the throbber and show the web browser content. */
+	WEBBROWSER_API void HideThrobber();
+
 private:
 
 	/** Navigate backwards. */
@@ -321,6 +327,9 @@ private:
 	/** The initial throbber setting */
 	bool bShowInitialThrobber;
 
+	/** Whether to force the throbber to appear */
+	bool bForceThrobber = false;
+
 	/** The WebBrowserWindow itself. */
 	TSharedPtr<IWebBrowserWindow> WebBrowserWindow;
 
