WebBrowser: Fix Korean input with Windows 11 IME

The "new" Windows 11 IME is also built on the TSF framework, but interacts differently with the app. For example, it doesn't call UpdateCompositionRange anymore. Another major difference is that it now expects to access the full string of text that was typed with IME rather than the last char only. As a result, some characters couldn't be entered correctly why typing with Korean IME.

Adjust FCEFTextInputMethodContext as following:
- keep track of the full composition sequence rather than the last composed char only, to handle multi-char compositions correctly (e.g. 'rkrk' with Korean IME)
- convert between textedit selection ranges and composition sequence ranges expected by TSF (otherwise it will sometimes end compositions early)
- handle manual caret changes during composition, so composition can be ended and a new sequence started (prevents leaving TSF in an inconsistent state)
- fix the appearance of the composed character(s) so they have the standard blue highlighting and underline
-remove unused and unnecessary calls which mess up the composition state on the CEF side

These changes are backwards compatible with the "old" Windows IME from Windows 10 and optionally restorable on Windows 11, and they also work on MacOS.

#tests verified with Korean, Japanese and Chinese IMEs on Win11 and MacOS, compared with native text edits
#rnx

#Synced-CL 49899484

diff --git a/Engine/Source/Runtime/WebBrowser/Private/CEF/CEFBrowserHandler.cpp b/Engine/Source/Runtime/WebBrowser/Private/CEF/CEFBrowserHandler.cpp
--- a/Engine/Source/Runtime/WebBrowser/Private/CEF/CEFBrowserHandler.cpp
+++ b/Engine/Source/Runtime/WebBrowser/Private/CEF/CEFBrowserHandler.cpp
@@ -421,6 +421,18 @@ void FCEFBrowserHandler::OnImeCompositionRangeChanged(
 		BrowserWindow->OnImeCompositionRangeChanged(Browser, SelectionRange, CharacterBounds);
 	}
 }
+
+void FCEFBrowserHandler::OnTextSelectionChanged(
+	CefRefPtr<CefBrowser> Browser,
+	const CefString& SelectedText,
+	const CefRange& SelectionRange)
+{
+	TSharedPtr<FCEFWebBrowserWindow> BrowserWindow = BrowserWindowPtr.Pin();
+	if (BrowserWindow.IsValid())
+	{
+		BrowserWindow->OnTextSelectionChanged(Browser, SelectedText, SelectionRange);
+	}
+}
 #endif
 
 
diff --git a/Engine/Source/Runtime/WebBrowser/Private/CEF/CEFBrowserHandler.h b/Engine/Source/Runtime/WebBrowser/Private/CEF/CEFBrowserHandler.h
--- a/Engine/Source/Runtime/WebBrowser/Private/CEF/CEFBrowserHandler.h
+++ b/Engine/Source/Runtime/WebBrowser/Private/CEF/CEFBrowserHandler.h
@@ -228,6 +228,10 @@ public:
 		CefRefPtr<CefBrowser> Browser,
 		const CefRange& SelectionRange,
 		const CefRenderHandler::RectList& CharacterBounds) override;
+	virtual void OnTextSelectionChanged(
+		CefRefPtr<CefBrowser> Browser,
+		const CefString& SelectedText,
+		const CefRange& SelectionRange) override;
 #endif
 
 public:
diff --git a/Engine/Source/Runtime/WebBrowser/Private/CEF/CEFImeHandler.cpp b/Engine/Source/Runtime/WebBrowser/Private/CEF/CEFImeHandler.cpp
--- a/Engine/Source/Runtime/WebBrowser/Private/CEF/CEFImeHandler.cpp
+++ b/Engine/Source/Runtime/WebBrowser/Private/CEF/CEFImeHandler.cpp
@@ -158,7 +158,6 @@ void FCEFImeHandler::UnbindCefBrowser()
 	{
 		DestroyContext();
 	}
-	InternalCefBrowser = nullptr;
 }
 
 
@@ -208,4 +207,16 @@ void FCEFImeHandler::CEFCompositionRangeChanged(const CefRange& SelectionRange,
 	}
 }
 
+void FCEFImeHandler::CEFTextSelectionChanged(const CefString& SelectedText, const CefRange& SelectionRange)
+{
+	if (TextInputMethodContext.IsValid() &&
+		TextInputMethodContext->CEFTextSelectionChanged(SelectedText, SelectionRange))
+	{
+		if (TextInputMethodChangeNotifier.IsValid())
+		{
+			TextInputMethodChangeNotifier->NotifySelectionChanged();
+		}
+	}
+}
+
 #endif
diff --git a/Engine/Source/Runtime/WebBrowser/Private/CEF/CEFImeHandler.h b/Engine/Source/Runtime/WebBrowser/Private/CEF/CEFImeHandler.h
--- a/Engine/Source/Runtime/WebBrowser/Private/CEF/CEFImeHandler.h
+++ b/Engine/Source/Runtime/WebBrowser/Private/CEF/CEFImeHandler.h
@@ -54,7 +54,14 @@ public:
 	* @param CharacterBounds The bounds of each character in view coordinates.
 	*/
 	void CEFCompositionRangeChanged(const CefRange& SelectionRange, const CefRenderHandler::RectList& CharacterBounds);
-	
+
+	/**
+	 * Called when the text selection has changed for a text field in a web page, used for IME.
+	 * @param SelectedText The selected text.
+	 * @param SelectionRange The range of characters that have been selected.
+	 */
+	void CEFTextSelectionChanged(const CefString& SelectedText, const CefRange& SelectionRange);
+
 	/**
 	 * Called when a message was received from the renderer process.
 	 *
diff --git a/Engine/Source/Runtime/WebBrowser/Private/CEF/CEFTextInputMethodContext.cpp b/Engine/Source/Runtime/WebBrowser/Private/CEF/CEFTextInputMethodContext.cpp
--- a/Engine/Source/Runtime/WebBrowser/Private/CEF/CEFTextInputMethodContext.cpp
+++ b/Engine/Source/Runtime/WebBrowser/Private/CEF/CEFTextInputMethodContext.cpp
@@ -18,7 +18,6 @@ FCEFTextInputMethodContext::FCEFTextInputMethodContext(const TSharedRef<FCEFImeH
 	: Owner(InOwner)
 	, bIsComposing(false)
 	, CompositionBeginIndex(0)
-	, CompositionLength(0)
 	, SelectionRangeBeginIndex(0)
 	, SelectionRangeLength(0)
 	, SelectionCaretPosition(ECaretPosition::Ending)
@@ -58,6 +57,25 @@ bool FCEFTextInputMethodContext::CEFCompositionRangeChanged(const CefRange& Sele
 	return false;
 }
 
+bool FCEFTextInputMethodContext::CEFTextSelectionChanged(const CefString& SelectedText, const CefRange& SelectionRange)
+{
+	// We get selection changes while composing as the caret follows the ongoing composition, these can be ignored.
+	// But any other selection change means the user clicked somewhere in the text field, and these must end any ongoing composition.
+	if (bIsComposing &&
+		(SelectionRangeBeginIndex == SelectionRange.from || SelectionRangeBeginIndex == (SelectionRange.from + CompositionString.Len())))
+	{
+		return false;
+	}
+
+	SelectionRangeBeginIndex = SelectionRange.from;
+	SelectionRangeLength = (SelectionRange.to - SelectionRange.from);
+
+	// Clicking anywhere in the textfield starts a new composition sequence
+	CompositionBeginIndex = SelectionRangeBeginIndex;
+	CompositionSequence.Empty();
+	return bIsComposing;
+}
+
 bool FCEFTextInputMethodContext::IsComposing()
 {
 	return bIsComposing;
@@ -70,85 +88,54 @@ bool FCEFTextInputMethodContext::IsReadOnly()
 
 uint32 FCEFTextInputMethodContext::GetTextLength()
 {
-	return CompositionString.Len();
+	return CompositionSequence.Len();
 }
 
 void FCEFTextInputMethodContext::GetSelectionRange(uint32& BeginIndex, uint32& Length, ECaretPosition& CaretPosition)
 {
-	BeginIndex = SelectionRangeBeginIndex;
+	BeginIndex = SelectionRangeBeginIndex - CompositionBeginIndex; // convert from textinput index to composition string index
 	Length = SelectionRangeLength;
 	CaretPosition = SelectionCaretPosition;
 }
 
 void FCEFTextInputMethodContext::SetSelectionRange(const uint32 BeginIndex, const uint32 Length, const ECaretPosition CaretPosition)
 {
-	SelectionRangeBeginIndex = BeginIndex;
+	SelectionRangeBeginIndex = BeginIndex + CompositionBeginIndex; // convert from composition string index to textinput index
 	SelectionRangeLength = Length;
 	SelectionCaretPosition = CaretPosition;
-
-	CefString Str = TCHAR_TO_WCHAR(*CompositionString);
-	std::vector<CefCompositionUnderline> underlines;
-
-	// Prevent CEF's ugly/unreadable default highlighting.
-	CefCompositionUnderline Underline;
-	{
-		Underline.range.from = 0;
-		Underline.range.to = Str.length();
-		Underline.color = 0;
-		Underline.background_color = 0;
-		Underline.thick = 0;
-		Underline.style = CEF_CUS_NONE;
-	}
-	underlines.push_back(Underline);
-	
-	Owner->InternalCefBrowser->GetHost()->ImeSetComposition(
-		Str,
-		underlines,
-		CefRange(UINT32_MAX, UINT32_MAX),
-		CefRange(SelectionRangeBeginIndex, SelectionRangeLength));
 }
 
 void FCEFTextInputMethodContext::GetTextInRange(const uint32 BeginIndex, const uint32 Length, FString& OutString)
 {
-	OutString = CompositionString.Mid(BeginIndex, Length);
+	OutString = CompositionSequence.Mid(BeginIndex, Length);
 }
 
 void FCEFTextInputMethodContext::SetTextInRange(const uint32 BeginIndex, const uint32 Length, const FString& InString)
 {
-	FString NewString;
+	FString NewCompositionSequence;
 	if (BeginIndex > 0)
 	{
-		NewString = CompositionString.Mid(0, BeginIndex);
+		NewCompositionSequence = CompositionSequence.Mid(0, BeginIndex);
 	}
 
-	NewString += InString;
+	NewCompositionSequence += InString;
 
-	if ((int32)(BeginIndex + Length) < CompositionString.Len())
+	if ((int32)(BeginIndex + Length) < CompositionSequence.Len())
 	{
-		NewString += CompositionString.Mid(BeginIndex + Length, CompositionString.Len() - (BeginIndex + Length));
+		NewCompositionSequence += CompositionSequence.Mid(BeginIndex + Length, CompositionSequence.Len() - (BeginIndex + Length));
 	}
-	CompositionString = NewString;
+	// this is the full sequence of chars typed with IME since we've clicked in the text field
+	CompositionSequence = NewCompositionSequence;
+	// this is the last char(s) that was typed with IME
+	CompositionString = InString;
 
 	CefString Str = TCHAR_TO_WCHAR(*CompositionString);
-	std::vector<CefCompositionUnderline> underlines;
 	
-	// Prevent CEF's ugly/unreadable default highlighting.
-	CefCompositionUnderline Underline;
-	{
-		Underline.range.from = 0;
-		Underline.range.to = Str.length();
-		Underline.color = 0;
-		Underline.background_color = 0;
-		Underline.thick = 0;
-		Underline.style = CEF_CUS_NONE;
-	}
-	underlines.push_back(Underline);
-
 	Owner->InternalCefBrowser->GetHost()->ImeSetComposition(
 		Str,
-		underlines,
+		{},
 		CefRange(UINT32_MAX, UINT32_MAX),
-		CefRange(0, Str.length()));
+		CefRange(0, CompositionString.Len())); // this selection range is within the composition string, not the whole text field!
 }
 
 int32 FCEFTextInputMethodContext::GetCharacterIndexFromPoint(const FVector2D& Point)
@@ -255,8 +242,6 @@ void FCEFTextInputMethodContext::BeginComposition()
 
 void FCEFTextInputMethodContext::UpdateCompositionRange(const int32 InBeginIndex, const uint32 InLength)
 {
-	CompositionBeginIndex = InBeginIndex;
-	CompositionLength = InLength;
 }
 
 void FCEFTextInputMethodContext::EndComposition()
@@ -282,11 +267,6 @@ void FCEFTextInputMethodContext::ResetComposition()
 {
 	CompositionString.Empty();
 	CefCompositionBounds.clear();
-	CompositionBeginIndex = 0;
-	CompositionLength = 0;
-
-	SelectionRangeBeginIndex = 0;
-	SelectionRangeLength = 0;
 }
 
 #endif
diff --git a/Engine/Source/Runtime/WebBrowser/Private/CEF/CEFTextInputMethodContext.h b/Engine/Source/Runtime/WebBrowser/Private/CEF/CEFTextInputMethodContext.h
--- a/Engine/Source/Runtime/WebBrowser/Private/CEF/CEFTextInputMethodContext.h
+++ b/Engine/Source/Runtime/WebBrowser/Private/CEF/CEFTextInputMethodContext.h
@@ -52,6 +52,7 @@ public:
 	bool UpdateCachedGeometry(const FGeometry& AllottedGeometry);
 
 	bool CEFCompositionRangeChanged(const CefRange& SelectionRange, const CefRenderHandler::RectList& CharacterBounds);
+	bool CEFTextSelectionChanged(const CefString& SelectedText, const CefRange& SelectionRange);
 
 private:
 	void ResetComposition();
@@ -84,7 +85,6 @@ private:
 	FGeometry CachedGeometry;
 	bool bIsComposing;
 	int32 CompositionBeginIndex;
-	uint32 CompositionLength;
 
 	uint32 SelectionRangeBeginIndex;
 	uint32 SelectionRangeLength;
@@ -92,6 +92,7 @@ private:
 
 	std::vector<CefRect> CefCompositionBounds;
 
+	FString CompositionSequence;
 	FString CompositionString;
 };
 
diff --git a/Engine/Source/Runtime/WebBrowser/Private/CEF/CEFWebBrowserWindow.cpp b/Engine/Source/Runtime/WebBrowser/Private/CEF/CEFWebBrowserWindow.cpp
--- a/Engine/Source/Runtime/WebBrowser/Private/CEF/CEFWebBrowserWindow.cpp
+++ b/Engine/Source/Runtime/WebBrowser/Private/CEF/CEFWebBrowserWindow.cpp
@@ -3024,6 +3024,19 @@ void FCEFWebBrowserWindow::OnImeCompositionRangeChanged(
 	}
 }
 
+void FCEFWebBrowserWindow::OnTextSelectionChanged(
+	CefRefPtr<CefBrowser> Browser,
+	const CefString& SelectedText,
+	const CefRange& SelectionRange)
+{
+	if (InternalCefBrowser != nullptr && InternalCefBrowser->IsSame(Browser))
+	{
+#if !PLATFORM_LINUX
+		Ime->CEFTextSelectionChanged(SelectedText, SelectionRange);
+#endif
+	}
+}
+
 void FCEFWebBrowserWindow::UpdateDragRegions(const TArray<FWebBrowserDragRegion>& Regions)
 {
 	DragRegions = Regions;
diff --git a/Engine/Source/Runtime/WebBrowser/Private/CEF/CEFWebBrowserWindow.h b/Engine/Source/Runtime/WebBrowser/Private/CEF/CEFWebBrowserWindow.h
--- a/Engine/Source/Runtime/WebBrowser/Private/CEF/CEFWebBrowserWindow.h
+++ b/Engine/Source/Runtime/WebBrowser/Private/CEF/CEFWebBrowserWindow.h
@@ -512,6 +512,17 @@ private:
 		const CefRange& SelectionRange,
 		const CefRenderHandler::RectList& CharacterBounds);
 
+	/**
+	 * Called when the text selection has changed for a text field in a web page, used for IME.
+	 * @param Browser The CefBrowser for this window.
+	 * @param SelectedText The selected text.
+	 * @param SelectionRange The range of characters that have been selected.
+	 */
+	void OnTextSelectionChanged(
+		CefRefPtr<CefBrowser> Browser,
+		const CefString& SelectedText,
+		const CefRange& SelectionRange);
+
 	void UpdateDragRegions(const TArray<FWebBrowserDragRegion>& Regions);
 
 public:
