Solving drawing in front conflicts in ApplePlatformWebBrowser through SlateRHIRenderer's Native layer/multi view rendering setup

#Synced-CL 50103067

diff --git a/Engine/Source/Runtime/WebBrowser/Private/Apple/ApplePlatformWebBrowser.cpp b/Engine/Source/Runtime/WebBrowser/Private/Apple/ApplePlatformWebBrowser.cpp
--- a/Engine/Source/Runtime/WebBrowser/Private/Apple/ApplePlatformWebBrowser.cpp
+++ b/Engine/Source/Runtime/WebBrowser/Private/Apple/ApplePlatformWebBrowser.cpp
@@ -348,6 +348,37 @@ class SAppleWebBrowserWidget : public SViewport
 			
 			UpdateFrameWithViewport(FIntPoint(static_cast<int32_t>(Size.X), static_cast<int32_t>(Size.Y)), 
 									FIntPoint(static_cast<int32_t>(Position.X), static_cast<int32_t>(Position.Y)));
+#if WITH_MULTI_VIEW_SLATE_WINDOW_SUPPORT
+			{
+				TSharedPtr<FWebBrowserWindow> WebBrowserWindow = WebBrowserWindowPtr.Pin();
+				if(WebBrowserWindow.IsValid())
+				{
+					TSharedPtr<SWindow> ParentWindow = WebBrowserWindow->GetParentWindow();
+					
+					if([WebViewWrapper isReadyToAddToNativeLayer] && LastSlateLayer.IsSet() && ParentWindow.IsValid())
+					{
+						int32 LayerId = LastSlateLayer.GetValue();
+						TSharedPtr<SWindow> PastNativeLayerWindow = LastNativeLayerCreatedInWindow.Pin();
+						bool bWasLayerAlreadyCreated = LastNativeLayerCreated.IsSet() && LastNativeLayerCreated.GetValue() == LayerId && 
+							PastNativeLayerWindow.IsValid() && PastNativeLayerWindow == ParentWindow;
+						
+						if(!bWasLayerAlreadyCreated)
+						{
+							UE_LOG(LogAppleWebBrowser, Display, TEXT("The native layer was not created for SAppleWebBrowserWidget %p on layer %d, creating..."), this, LayerId);
+
+							if(LastNativeLayerCreated.IsSet() && PastNativeLayerWindow.IsValid())
+							{
+								FSlateApplication::Get().GetRenderer()->DeleteNativeLayer(LastNativeLayerCreated.GetValue(), *PastNativeLayerWindow);
+							}
+							
+							FSlateApplication::Get().GetRenderer()->CreateNativeLayer(LayerId, *ParentWindow, WebViewWrapper.WebViewContainer);
+							LastNativeLayerCreated = LayerId;
+							LastNativeLayerCreatedInWindow = ParentWindow;
+						}
+					}
+				}
+			}
+#endif
 			
 #if !PLATFORM_TVOS
 			if ([WebViewWrapper isAllowingScreenshots] || UsingIndirectRendering())
@@ -425,6 +456,9 @@ class SAppleWebBrowserWidget : public SViewport
 			FSlateDrawElement::MakeBox(OutDrawElements, LayerId, AllottedGeometry.ToPaintGeometry(), WebBrowserBrush.Get(), ESlateDrawEffect::None);
 		}
 #endif
+
+		LastSlateLayer = LayerId;
+
 		return LayerId;
 	}
 
@@ -467,6 +501,19 @@ class SAppleWebBrowserWidget : public SViewport
 
 	void Close()
 	{
+		if(LastNativeLayerCreated.IsSet())
+		{
+			TSharedPtr<SWindow> NativeLayerWindow = LastNativeLayerCreatedInWindow.Pin();
+			
+			if(NativeLayerWindow.IsValid())
+			{
+				FSlateApplication::Get().GetRenderer()->DeleteNativeLayer(LastNativeLayerCreated.GetValue(), 
+																		  *NativeLayerWindow);
+			}
+			
+			LastNativeLayerCreated.Reset();
+		}
+		
 #if !PLATFORM_TVOS
 		FCoreDelegates::OnSafeFrameChangedEvent.Remove(OnSafeFrameChangedEventHandle);
 #endif
@@ -744,6 +791,11 @@ private:
 	bool bForceRecomputeWindowOnPaint = false;
 	bool bTakingScreenshot = false;
 	bool bFirstScreenshotTaken = false;
+	
+	TOptional<int32> LastNativeLayerCreated;
+	TWeakPtr<SWindow> LastNativeLayerCreatedInWindow;
+	
+	mutable TOptional<int32> LastSlateLayer;
 
 #if !PLATFORM_TVOS
 	/** The external texture to render the webbrowser output. */
@@ -800,6 +852,7 @@ private:
 	Is3DBrowser = false;
 	bVideoTextureValid = false;
 	bShuttingDown = false;
+	bReadyToAddToNativeLayer = false;
 	
 #if !UE_BUILD_SHIPPING
 	BOOL Inspectable = CVarWebViewEnableInspectable.GetValueOnAnyThread()? YES : NO;
@@ -986,10 +1039,8 @@ private:
 
 -(void)leaveDeadWindow;
 {
-	if(!bNeedsAddToView && !Is3DBrowser) 
-	{
-		bNeedsAddToView = true;
-	}
+	bNeedsAddToView = true;
+	bReadyToAddToNativeLayer = false;
 }
 
 -(void)updateframe:(CGRect)InFrame window:(FGenericWindow*)EnclosedWindow;
@@ -1061,8 +1112,6 @@ private:
 #endif
 			if (bNeedsAddToView)
 			{
-				bNeedsAddToView = false;
-				
 #if PLATFORM_IOS 
 				[self.WebViewContainer setOpaque : NO];
 				[self.WebViewContainer setBackgroundColor : [UIColor clearColor]];
@@ -1072,12 +1121,14 @@ private:
 				{
 					[CloseButton setupLayout];
 				}
-#else				
+#elif !WITH_MULTI_VIEW_SLATE_WINDOW_SUPPORT				
 				FCocoaWindow *CocoaWindow = ((FMacWindow*)EnclosedWindow)->GetWindowHandle();
-				NSView *CocoaWindowView = CocoaWindow.openGLView;
+				NSView *CocoaWindowView = [CocoaWindow contentView];
 				
 				[CocoaWindowView addSubview : self.WebViewContainer];
 #endif
+				bNeedsAddToView = false;
+				bReadyToAddToNativeLayer = true;
 			}
 			else
 			{
@@ -1484,6 +1535,13 @@ private:
 #endif
 }
 
+#if WITH_MULTI_VIEW_SLATE_WINDOW_SUPPORT
+-(bool)isReadyToAddToNativeLayer;
+{
+	return bReadyToAddToNativeLayer;
+}
+#endif
+
 #if !PLATFORM_TVOS
 
 - (nullable WKWebView *)webView:(WKWebView *)InWebView createWebViewWithConfiguration:(WKWebViewConfiguration *)InConfiguration forNavigationAction:(WKNavigationAction *)InNavigationAction windowFeatures:(WKWindowFeatures *)InWindowFeatures
@@ -1525,7 +1583,7 @@ private:
 	else
 	{
 		UE_LOG(LogAppleWebBrowser, Warning, 
-			   TEXT("ApplePlatformWebBrowser recieved decidePolicyForNavigationAction callback during or after module shutdown, preventing navigation"));
+			   TEXT("ApplePlatformWebBrowser received decidePolicyForNavigationAction callback during or after module shutdown, preventing navigation"));
 		InDecisionHandler(WKNavigationActionPolicyCancel);
 	}
 }
@@ -1544,7 +1602,7 @@ private:
 	else
 	{
 		UE_LOG(LogAppleWebBrowser, Warning, 
-			   TEXT("ApplePlatformWebBrowser recieved didCommitNavigation callback during or after module shutdown, ignoring"));
+			   TEXT("ApplePlatformWebBrowser received didCommitNavigation callback during or after module shutdown, ignoring"));
 	}
 }
 
@@ -1561,7 +1619,7 @@ private:
 	else
 	{
 		UE_LOG(LogAppleWebBrowser, Warning, 
-			   TEXT("ApplePlatformWebBrowser recieved didFinishNavigation callback during or after module shutdown, ignoring"));
+			   TEXT("ApplePlatformWebBrowser received didFinishNavigation callback during or after module shutdown, ignoring"));
 	}
 }
 -(void)webView:(WKWebView *)InWebView didFailNavigation : (WKNavigation *)InNavigation withError : (NSError*)InError
@@ -1580,7 +1638,7 @@ private:
 	else
 	{
 		UE_LOG(LogAppleWebBrowser, Warning, 
-			   TEXT("ApplePlatformWebBrowser recieved didFailNavigation callback during or after module shutdown, ignoring"));
+			   TEXT("ApplePlatformWebBrowser received didFailNavigation callback during or after module shutdown, ignoring"));
 	}
 }
 -(void)webView:(WKWebView *)InWebView didFailProvisionalNavigation : (WKNavigation *)InNavigation withError : (NSError*)InError
@@ -1594,7 +1652,7 @@ private:
 	else
 	{
 		UE_LOG(LogAppleWebBrowser, Warning, 
-			   TEXT("ApplePlatformWebBrowser recieved didFailProvisionalNavigation callback during or after module shutdown, ignoring"));
+			   TEXT("ApplePlatformWebBrowser received didFailProvisionalNavigation callback during or after module shutdown, ignoring"));
 	}
 }
 #endif
diff --git a/Engine/Source/Runtime/WebBrowser/Private/Apple/ApplePlatformWebBrowser.h b/Engine/Source/Runtime/WebBrowser/Private/Apple/ApplePlatformWebBrowser.h
--- a/Engine/Source/Runtime/WebBrowser/Private/Apple/ApplePlatformWebBrowser.h
+++ b/Engine/Source/Runtime/WebBrowser/Private/Apple/ApplePlatformWebBrowser.h
@@ -55,6 +55,7 @@ class SWebBrowserView;
 	// but read from AppKit/UIKit's UI thread during WebKit
 	// callbacks
 	std::atomic<bool> bShuttingDown;
+	bool bReadyToAddToNativeLayer;
 }
 #if !PLATFORM_TVOS
 @property(strong) WKWebView* WebView;
@@ -110,6 +111,9 @@ enableFloatingCloseButton : (bool)bEnableFloatingCloseButton;
 -(void)updateWebViewMetalTexture : (id<MTLTexture>)texture;
 -(void)leaveDeadWindow;
 -(bool)inKeyWindow;
+#if WITH_MULTI_VIEW_SLATE_WINDOW_SUPPORT
+-(bool)isReadyToAddToNativeLayer;
+#endif
 #if !PLATFORM_MAC
 -(CGFloat)getScreenScale;
 #endif
