Refactored ApplePlatformWebBrowser Slate Reflector/Screenshot support to only take a WebKit screenshot when taking a Slate screenshot

#Synced-CL 50166977

diff --git a/Engine/Source/Runtime/WebBrowser/Private/Apple/ApplePlatformWebBrowser.cpp b/Engine/Source/Runtime/WebBrowser/Private/Apple/ApplePlatformWebBrowser.cpp
--- a/Engine/Source/Runtime/WebBrowser/Private/Apple/ApplePlatformWebBrowser.cpp
+++ b/Engine/Source/Runtime/WebBrowser/Private/Apple/ApplePlatformWebBrowser.cpp
@@ -53,7 +53,7 @@ class SAppleWebBrowserWidget : public SViewport
 	SLATE_BEGIN_ARGS(SAppleWebBrowserWidget)
 		: _InitialURL("about:blank")
 		, _UseTransparency(false)
-		, _AllowScreenshots(false)
+		, _AllowScreenshots(true)
 	{ }
 
 	SLATE_ARGUMENT(FString, InitialURL);
@@ -228,6 +228,13 @@ class SAppleWebBrowserWidget : public SViewport
 				WindowFrame = NewFrame;
 				PastEnclosedWindow = EnclosedWindow;
 				LeaveDeadWindow();
+				
+				// Need to also replace the widget path as otherwise
+				// screenshots will break when moving windows
+				if(CurrentWidgetPath.IsValid())
+				{
+					CurrentWidgetPath = FWidgetPath{};
+				}
 			}
 		}
 #else
@@ -259,9 +266,50 @@ class SAppleWebBrowserWidget : public SViewport
 #endif
 	}
 	
+	static void CreateWebBrowserTextureIfNeeded(FRHICommandListImmediate& RHICmdList, 
+												AppleWebViewWrapper* NativeWebBrowser, FIntPoint Size)
+	{
+		check(IsInRenderingThread());
+		
+		FTextureRHIRef VideoTexture = [NativeWebBrowser GetVideoTexture];
+		bool bStaleVideoTexture = VideoTexture != nullptr && VideoTexture->GetSizeXY() != Size;
+									
+		if (VideoTexture == nullptr || bStaleVideoTexture)
+		{
+			// We need this texture to be accessible and writable from the CPU because we need to call
+			// replaceRegion:mipmapLevel:withBytes:bytesPerRow: to copy over a screenshot captured in WKWebView
+			// using takeSnapshotWithConfiguration:completionHandler: which gives us an UI/NSImage.
+			const FRHITextureCreateDesc Desc =
+				FRHITextureCreateDesc::Create2D(TEXT("SAppleWebBrowserWidget_VideoTexture"), Size, PF_R8G8B8A8)
+				.SetFlags(ETextureCreateFlags::CPUWritable | ETextureCreateFlags::CPUReadback);
+
+			VideoTexture = RHICmdList.CreateTexture(Desc);
+			[NativeWebBrowser SetVideoTexture : VideoTexture];
+
+			if (VideoTexture == nullptr)
+			{
+				UE_LOG(LogAppleWebBrowser, Warning, TEXT("CreateTexture failed!"));
+				return;
+			}
+
+			[NativeWebBrowser SetVideoTextureValid : false];
+		}
+	}
+	
 	void WriteWebBrowserTexture(AppleWebViewWrapper* NativeWebBrowserPtr,
 								FGuid PlayerGuid, FIntPoint Size)
 	{
+		// The FlushRenderingCommands for screenshot mode are needed to make
+		// taking a screenshot a sync op so we don't need screenshot info from the previous frame.
+		if(bTakingScreenshot)
+		{
+			ENQUEUE_RENDER_COMMAND(CreateWebBrowserTexture)([NativeWebBrowserPtr, PlayerGuid, Size](FRHICommandListImmediate& RHICmdList) 
+															{
+				CreateWebBrowserTextureIfNeeded(RHICmdList, NativeWebBrowserPtr, Size);
+			});
+			FlushRenderingCommands();
+		}
+		
 		ENQUEUE_RENDER_COMMAND(WriteWebBrowser)(
 			[NativeWebBrowserPtr, PlayerGuid, Size](FRHICommandListImmediate& RHICmdList)
 			{
@@ -272,31 +320,14 @@ class SAppleWebBrowserWidget : public SViewport
 					return;
 				}
 
+				CreateWebBrowserTextureIfNeeded(RHICmdList, NativeWebBrowserPtr, Size);
 				FTextureRHIRef VideoTexture = [NativeWebBrowser GetVideoTexture];
-				bool bStaleVideoTexture = VideoTexture != nullptr && VideoTexture->GetSizeXY() != Size;
-											
-				if (VideoTexture == nullptr || bStaleVideoTexture)
+				
+				if(VideoTexture == nil)
 				{
-					// We need this texture to be accessible and writable from the CPU because we need to call
-					// replaceRegion:mipmapLevel:withBytes:bytesPerRow: to copy over a screenshot captured in WKWebView
-					// using takeSnapshotWithConfiguration:completionHandler: which gives us an UI/NSImage.
-					const FRHITextureCreateDesc Desc =
-						FRHITextureCreateDesc::Create2D(TEXT("SAppleWebBrowserWidget_VideoTexture"), Size, PF_R8G8B8A8)
-						.SetFlags(ETextureCreateFlags::CPUWritable | ETextureCreateFlags::CPUReadback);
-
-					VideoTexture = RHICmdList.CreateTexture(Desc);
-					[NativeWebBrowser SetVideoTexture : VideoTexture];
-
-					if (VideoTexture == nullptr)
-					{
-						UE_LOG(LogAppleWebBrowser, Warning, TEXT("CreateTexture failed!"));
-						return;
-					}
-
-					[NativeWebBrowser SetVideoTextureValid : false];
-
+					return;
 				}
-
+				
 				if ([NativeWebBrowser UpdateVideoFrame : VideoTexture->GetNativeResource()])
 				{
 					// if region changed, need to reregister UV scale/offset
@@ -312,6 +343,11 @@ class SAppleWebBrowserWidget : public SViewport
 					[NativeWebBrowser SetVideoTextureValid : true];
 				}
 			});
+		
+		if(bTakingScreenshot)
+		{
+			FlushRenderingCommands();
+		}
 	}
 	
 	void Tick(const FGeometry& AllottedGeometry, const double InCurrentTime, const float InDeltaTime)
@@ -392,7 +428,11 @@ class SAppleWebBrowserWidget : public SViewport
 					bTakingScreenshot = true;
 					bFirstScreenshotTaken = true;
 				}
-				
+				else if(bFirstScreenshotTaken && !Is3DBrowser && !bTakingScreenshot)
+				{
+					// we shouldn't be switching to the expensive screenshot mode when we don't need to!
+					return;
+				}
 				
 				if (WebBrowserTexture.IsValid())
 				{					
@@ -817,7 +857,7 @@ private:
 	FDelegateHandle OnSafeFrameChangedEventHandle;
 #endif
 	
-	FWidgetPath CurrentWidgetPath;
+	FWidgetPath CurrentWidgetPath{};
 	
 	CGRect WindowFrame;
 };
@@ -1447,6 +1487,7 @@ private:
 -(void)updateWebViewMetalTexture:(id<MTLTexture>)texture
 {
 #if !PLATFORM_TVOS
+	dispatch_semaphore_t semaphore = dispatch_semaphore_create(0);
 	
 	dispatch_sync(dispatch_get_main_queue(), ^ 
 	{
@@ -1459,7 +1500,12 @@ private:
 #if PLATFORM_MAC
 			NSUInteger Width = TextureWidth;
 			NSUInteger Height = TextureHeight;
-			WebView.frame = CGRectMake(0, 0, Width, Height);
+			if(Is3DBrowser)
+			{
+				// This should only be the authoritative frame if we are a 3D browser
+				// and we have no other UI constraints forced on us to worry about
+				WebView.frame = CGRectMake(0, 0, Width, Height);
+			}
 #else
 			float ScreenScale = [self getScreenScale];
 			NSUInteger Width = TextureWidth / ScreenScale;
@@ -1503,19 +1549,32 @@ private:
 				
 				Image = nil;
 				[SnapshotConfig release];
-			}];
-			
-			if(!Is3DBrowser && [self isUsingIndirectRendering])
-			{
-				// always take screenshots to get around async weirdness but never 
-				// do an expensive copy to the Slate texture unless we need to because
-				// we are actively taking a screenshot this frame.
 				
-				check(bAllowScreenshots);
-				[self replaceMetalTexture: texture fromImage: LastWebViewScreenshot];
-			}
+				if(!Is3DBrowser && [self isUsingIndirectRendering])
+				{
+					// This semaphore only matters if we are trying to do a synchronous copy of 
+					// WKWebView's content, which is only necessary when we are trying to take a
+					// screenshot.
+					dispatch_semaphore_signal(semaphore);
+				}
+			}];
 		}
 	});
+	
+	if(!Is3DBrowser && [self isUsingIndirectRendering])
+	{
+		// Wait on the render thread for the screenshot to complete on the main thread.
+		// Afterwards, we wait to copy that screenshot over to the metal texture to
+		// make this one synchronous operation. This makes Slate screenshotting work as
+		// expected without needing a WebKit screenshot from a previous frame.
+		dispatch_semaphore_wait(semaphore, DISPATCH_TIME_FOREVER);
+			
+		dispatch_sync(dispatch_get_main_queue(), ^ 
+		{
+			check(bAllowScreenshots);
+			[self replaceMetalTexture: texture fromImage: LastWebViewScreenshot];
+		});
+	}
 #endif
 }
 
