Remove string copies under SWebBrowser/SWebBrowserView::LoadURL/LoadString

#rnx

#Synced-CL 50529789

diff --git a/Engine/Source/Runtime/WebBrowser/Private/SWebBrowser.cpp b/Engine/Source/Runtime/WebBrowser/Private/SWebBrowser.cpp
--- a/Engine/Source/Runtime/WebBrowser/Private/SWebBrowser.cpp
+++ b/Engine/Source/Runtime/WebBrowser/Private/SWebBrowser.cpp
@@ -321,7 +321,7 @@ void SWebBrowser::LoadURL(FString NewURL)
 {
 	if (BrowserView.IsValid())
 	{
-		BrowserView->LoadURL(NewURL);
+		BrowserView->LoadURL(MoveTemp(NewURL));
 	}
 }
 
@@ -329,7 +329,7 @@ void SWebBrowser::LoadString(FString Contents, FString DummyURL)
 {
 	if (BrowserView.IsValid())
 	{
-		BrowserView->LoadString(Contents, DummyURL);
+		BrowserView->LoadString(MoveTemp(Contents), MoveTemp(DummyURL));
 	}
 }
 
diff --git a/Engine/Source/Runtime/WebBrowser/Private/SWebBrowserView.cpp b/Engine/Source/Runtime/WebBrowser/Private/SWebBrowserView.cpp
--- a/Engine/Source/Runtime/WebBrowser/Private/SWebBrowserView.cpp
+++ b/Engine/Source/Runtime/WebBrowser/Private/SWebBrowserView.cpp
@@ -312,7 +312,7 @@ void SWebBrowserView::LoadURL(FString NewURL)
 	AddressBarUrl = FText::FromString(NewURL);
 	if (BrowserWindow.IsValid())
 	{
-		BrowserWindow->LoadURL(NewURL);
+		BrowserWindow->LoadURL(MoveTemp(NewURL));
 	}
 }
 
@@ -320,7 +320,7 @@ void SWebBrowserView::LoadString(FString Contents, FString DummyURL)
 {
 	if (BrowserWindow.IsValid())
 	{
-		BrowserWindow->LoadString(Contents, DummyURL);
+		BrowserWindow->LoadString(MoveTemp(Contents), MoveTemp(DummyURL));
 	}
 }
 
