WebBrowser: improve performance during resizing with CEF

- Reworked the resizing logic to avoid stretching the old webview contents until the new resized frames come in from CEF. This greatly improves the visual feedback during resizes.
- Improve the smooth sizing behavior to reduce overhead and avoid spamming CEF with resize requests
- Trigger resize requests to CEF right away instead of delaying to the next Tick. The request is already async so there's no need to add further delay.
- Remove the hack to artificially tick CEF when a floating window is being resized. The hack caused several state issues and contributed to resizing being slow and laggy. This is an issue with the UE window resizing stopping the game ticker for some reason, if they fix it will fix it for CEF too. It's not worth the trouble of trying to circumvent it.
- Found several optimizations to reduce the time needed to copy the frames we receive from CEF into our cached textures

#rnx

#Synced-CL 50345329

diff --git a/Engine/Source/Runtime/WebBrowser/Private/CEF/CEFBrowserApp.cpp b/Engine/Source/Runtime/WebBrowser/Private/CEF/CEFBrowserApp.cpp
--- a/Engine/Source/Runtime/WebBrowser/Private/CEF/CEFBrowserApp.cpp
+++ b/Engine/Source/Runtime/WebBrowser/Private/CEF/CEFBrowserApp.cpp
@@ -136,13 +136,20 @@ bool FCEFBrowserApp::TickMessagePump(float DeltaTime, bool bForce)
 		++NumFrames;
 		if (NumFrames % 100 == 0 || SecondsFrameRate > 5.0f)
 		{
-			UE_LOG(LogWebBrowser, Error, TEXT("CefDoMessageLoopWork call Frame Rate %0.2f"), NumFrames / SecondsFrameRate);
+			UE_LOG(LogCEFBrowser, Log, TEXT("CefDoMessageLoopWork call Frame Rate %0.2f"), NumFrames / std::max(SecondsFrameRate, 0.01f));
 			SecondsFrameRate = 0;
 			NumFrames = 0;
 		}
-#endif
 
+		static double LastTickTime = FPlatformTime::Seconds();
+		double BeforeSec = FPlatformTime::Seconds();
+#endif
 		CefDoMessageLoopWork();
+#ifdef DEBUG_CEFMESSAGELOOP_FRAMERATE
+		double AfterSec = FPlatformTime::Seconds();
+		UE_LOG(LogCEFBrowser, Verbose, TEXT("FCEFBrowserApp::TickMessagePump %u ms %u ms %s"), static_cast<uint32_t>((BeforeSec - LastTickTime) * 1000), static_cast<uint32_t>((AfterSec - BeforeSec) * 1000), bForce ? TEXT("FORCED") : TEXT(""));
+		LastTickTime = BeforeSec;
+#endif
 		return true;
 	}
 	return false;
diff --git a/Engine/Source/Runtime/WebBrowser/Private/CEF/CEFWebBrowserWindow.cpp b/Engine/Source/Runtime/WebBrowser/Private/CEF/CEFWebBrowserWindow.cpp
--- a/Engine/Source/Runtime/WebBrowser/Private/CEF/CEFWebBrowserWindow.cpp
+++ b/Engine/Source/Runtime/WebBrowser/Private/CEF/CEFWebBrowserWindow.cpp
@@ -244,6 +244,10 @@ enum KeyboardCode {
 #endif
 
 #define UE_CEF_HAS_RESIZE_BUG (CEF_VERSION_MAJOR >= 128)
+// CEF is usually 1-2 frames behind our current viewport size when drag-resizing,
+// this value gives it just enough time to send a correctly-sized frame without us forcing a resize
+// (which invalidates a bunch of things on the CEF side)
+#define MIN_FRAMES_FOR_CEF_RESIZE 3
 
 namespace {
 	// Private helper class to post a callback to GetSource.
@@ -398,7 +402,6 @@ FCEFWebBrowserWindow::FCEFWebBrowserWindow(CefRefPtr<CefBrowser> InBrowser, CefR
 	, bIsDisabled(false)
 	, bIsHidden(false)
 	, bTickedLastFrame(true)
-	, bNeedsResize(false)
 	, bDraggingWindow(false)
 	, PreviousKeyDownEvent()
 	, PreviousKeyUpEvent()
@@ -423,6 +426,9 @@ FCEFWebBrowserWindow::FCEFWebBrowserWindow(CefRefPtr<CefBrowser> InBrowser, CefR
 #if PLATFORM_WINDOWS || PLATFORM_MAC
 	, bInDirectHwndMode(false)
 #endif
+	, bHasCorrectNativeCefBuffer(true)
+	, bViewportResizePending(false)
+	, NumFramesSinceViewportResize(0)
 {
 	check(InBrowser.get() != nullptr);
 	check(!bUsingAcceleratedPaint || CanSupportAcceleratedPaint()); // make sure if accelerated paint is selected we can support it
@@ -634,9 +640,9 @@ void FCEFWebBrowserWindow::SetViewportSize(FIntPoint WindowSize, FIntPoint Windo
 	// Ignore sizes that can't be seen as it forces CEF to re-render whole image
 	if ((WindowSize.X > 0 && WindowSize.Y > 0 && ViewportSize != WindowSize) || WindowDPIScaleFactor != ViewportDPIScaleFactor)
 	{
-		bool bFirstSize = ViewportSize == FIntPoint::ZeroValue;
 		ViewportSize = WindowSize;
 		ViewportDPIScaleFactor = WindowDPIScaleFactor;
+		FIntPoint WindowSizeScaled = (FVector2D(WindowSize) * WindowDPIScaleFactor).IntPoint();
 
 		if (IsValid())
 		{
@@ -652,8 +658,6 @@ void FCEFWebBrowserWindow::SetViewportSize(FIntPoint WindowSize, FIntPoint Windo
 					::GetWindowRect(Parent, &ParentRect);
 				}
 
-				FIntPoint WindowSizeScaled = (FVector2D(WindowSize) * WindowDPIScaleFactor).IntPoint();
-
 				::SetWindowPos(NativeHandle, 0, WindowPos.X - ParentRect.left, WindowPos.Y - ParentRect.top, WindowSizeScaled.X, WindowSizeScaled.Y, 0);
 			}
 #elif PLATFORM_MAC
@@ -674,14 +678,13 @@ void FCEFWebBrowserWindow::SetViewportSize(FIntPoint WindowSize, FIntPoint Windo
 				}
 			}
 #endif
+			bHasCorrectNativeCefBuffer = false;
+			NumFramesSinceViewportResize = 0;
+			RequestCefResize();
 
-			if (bFirstSize)
-			{
-				InternalCefBrowser->GetHost()->WasResized();
-			}
-			else
+			if (RHIRenderHelper)
 			{
-				bNeedsResize = true;
+				UpdatableTextures[PET_VIEW] = RHIRenderHelper->ResizeSlateUpdatableTexture(UpdatableTextures[PET_VIEW], WindowSizeScaled);
 			}
 		}
 	}
@@ -2087,6 +2090,13 @@ void FCEFWebBrowserWindow::OnAcceleratedPaint(CefRenderHandler::PaintElementType
 		}
 	}
 
+	FIntRect DirtyRect = (DirtyRects.size() == 1) ? FIntRect(DirtyRects[0].x, DirtyRects[0].y, DirtyRects[0].x + DirtyRects[0].width, DirtyRects[0].y + DirtyRects[0].height) : FIntRect();
+	FIntRect ViewportRect = FIntRect(0, 0, ViewportSize.X, ViewportSize.Y).Scale(ViewportDPIScaleFactor);
+	if (DirtyRect.IsEmpty() || ViewportRect.IsEmpty())
+	{
+		return;
+	}
+
 #if CEF_VERSION_MAJOR >= 128
 #  if PLATFORM_WINDOWS
 	void* SharedHandle = Info.shared_texture_handle;
@@ -2105,14 +2115,25 @@ void FCEFWebBrowserWindow::OnAcceleratedPaint(CefRenderHandler::PaintElementType
 	LastPaintedSharedHandle = SharedHandle;
 #endif
 
-	FIntRect DirtyRect = (DirtyRects.size() == 1) ? FIntRect(DirtyRects[0].x, DirtyRects[0].y, DirtyRects[0].x + DirtyRects[0].width, DirtyRects[0].y + DirtyRects[0].height) : FIntRect();
-	FIntRect ViewportRect = FIntRect(0, 0, ViewportSize.X, ViewportSize.Y).Scale(ViewportDPIScaleFactor);
+	bool bWasCorrectNativeCefBuffer = bHasCorrectNativeCefBuffer;
 	bHasCorrectNativeCefBuffer = (DirtyRect == ViewportRect);
+	bViewportResizePending = false;
+	++NumFramesSinceViewportResize;
+	if (!bHasCorrectNativeCefBuffer)
+	{
+		UE_LOG(LogWebBrowser, Verbose, TEXT("FCEFWebBrowserWindow::OnAcceleratedPaint: Frame size still wrong after %zu frames: %dx%d -> %dx%d (%d, %d)"), NumFramesSinceViewportResize, ViewportRect.Width(), ViewportRect.Height(), DirtyRect.Width(), DirtyRect.Height(), DirtyRect.Min.X, DirtyRect.Min.Y);
+	}
+	else if (!bWasCorrectNativeCefBuffer)
+	{
+		UE_LOG(LogWebBrowser, Verbose, TEXT("FCEFWebBrowserWindow::OnAcceleratedPaint: Frame size OK again after %zu frames: %dx%d -> %dx%d (%d, %d)"), NumFramesSinceViewportResize, ViewportRect.Width(), ViewportRect.Height(), DirtyRect.Width(), DirtyRect.Height(), DirtyRect.Min.X, DirtyRect.Min.Y);
+	}
 
 	// During resizing, CEF occasionally sends textures with black bars either above/below or left/right and distorted content in between, we want to ignore these textures.
 	// Fortunately, the DirtyRect for these textures (and only these) has a non-zero X/Y depending on the black bars orientation, so we can detect them without inspecting the content.
 	if (DirtyRect.Min.X > 0 || DirtyRect.Min.Y > 0)
 	{
+		// force a CEF resize even if we were below the threshold, as CEF considers those "good" frames and will stop sending more until requested
+		NumFramesSinceViewportResize = MIN_FRAMES_FOR_CEF_RESIZE;
 		return;
 	}
 
@@ -2667,22 +2688,6 @@ void FCEFWebBrowserWindow::CheckTickActivity()
 	{
 		SetIsHidden(true);
 	}
-	else if(bNeedsResize)
-	{
-		bNeedsResize = false;
-		InternalCefBrowser->GetHost()->WasResized();
-
-		// Recent CEF versions fail to trigger a repaint after WasResized() is called
-		// so we must force one on the view (and the popup if visible):
-		// (upstream issue: [REDACTED_PROFILE_URL]/cef/issues/3826)
-#if UE_CEF_HAS_RESIZE_BUG
-		InternalCefBrowser->GetHost()->Invalidate(PET_VIEW);
-		if (bPopupHasFocus)
-		{
-			InternalCefBrowser->GetHost()->Invalidate(PET_POPUP);
-		}
-#endif
-	}
 
 	bTickedLastFrame = false;
 }
@@ -2906,9 +2911,36 @@ TSharedPtr<FCapturedCefBuffer> FCEFWebBrowserWindow::GetCapturedCefBuffer()
 	return CapturedCefBuffer; // ..may not be set
 }
 
-void FCEFWebBrowserWindow::SetNeedsResize(const bool bInValue)
+bool FCEFWebBrowserWindow::IsCefResizeNeeded() const
+{
+	if (!bHasCorrectNativeCefBuffer && !bViewportResizePending && NumFramesSinceViewportResize >= MIN_FRAMES_FOR_CEF_RESIZE)
+	{
+		UE_LOG(LogWebBrowser, Verbose, TEXT("FCEFWebBrowserWindow::IsCefResizeNeeded: Forcing CEF resize"));
+		return true;
+	}
+	return false;
+}
+
+void FCEFWebBrowserWindow::RequestCefResize()
 {
-	bNeedsResize = bInValue;
+	if (!IsValid())
+	{
+		return;
+	}
+
+	InternalCefBrowser->GetHost()->WasResized();
+	// Recent CEF versions fail to trigger a repaint after WasResized() is called
+	// so we must force one on the view (and the popup if visible):
+	// (upstream issue: [REDACTED_PROFILE_URL]/cef/issues/3826)
+#if UE_CEF_HAS_RESIZE_BUG
+	InternalCefBrowser->GetHost()->Invalidate(PET_VIEW);
+	if (bPopupHasFocus)
+	{
+		InternalCefBrowser->GetHost()->Invalidate(PET_POPUP);
+	}
+#endif
+
+	bViewportResizePending = true;
 }
 
 bool FCEFWebBrowserWindow::OnProcessMessageReceived(CefRefPtr<CefBrowser> Browser, CefRefPtr<CefFrame> frame, CefProcessId SourceProcess, CefRefPtr<CefProcessMessage> Message)
diff --git a/Engine/Source/Runtime/WebBrowser/Private/CEF/CEFWebBrowserWindow.h b/Engine/Source/Runtime/WebBrowser/Private/CEF/CEFWebBrowserWindow.h
--- a/Engine/Source/Runtime/WebBrowser/Private/CEF/CEFWebBrowserWindow.h
+++ b/Engine/Source/Runtime/WebBrowser/Private/CEF/CEFWebBrowserWindow.h
@@ -565,7 +565,6 @@ public:
 	 */
 	void UpdateCachedGeometry(const FGeometry& AllottedGeometry);
 
-
 	/**
 	 * Called from the WebBrowserSingleton tick event. Should test wether the widget got a tick from Slate last frame and set the state to hidden if not.
 	 */
@@ -590,28 +589,33 @@ public:
 	/**
 	 * Tells us whether this has the correct CEF image buffer.
 	 */
-	WEBBROWSER_API bool HasCorrectNativeCefBuffer() const;
+	bool HasCorrectNativeCefBuffer() const;
 	
 	/**
 	 * Whether we should use a workaround for this bug -
 	 * [REDACTED_PROFILE_URL]/cef/issues/3826
 	 */
-	WEBBROWSER_API void SetUseCapturedCefBuffer(const bool bValue);
+	void SetUseCapturedCefBuffer(const bool bValue);
 	
 	/**
 	 * Whether the CEF version we're using has a resize bug that sometimes requires requesting new frames after a resize.
 	 */
-	WEBBROWSER_API bool HasResizeBug() const;
+	bool HasResizeBug() const;
 	
 	/**
 	 * Gets an undistorted buffer captured from CEF. Format is PF_B8G8R8A8.
 	 */
-	WEBBROWSER_API TSharedPtr<FCapturedCefBuffer> GetCapturedCefBuffer();
+	TSharedPtr<FCapturedCefBuffer> GetCapturedCefBuffer();
+
+	/**
+	 * Indicates a CEF resize is needed, using heuristics to avoid too many resize requests during an active drag-resize.
+	 */
+	bool IsCefResizeNeeded() const;
 
 	/**
-	 * Set if a resize is needed.
+	 * Requests a resize of the CEF webview to ViewportSize
 	 */
-	WEBBROWSER_API void SetNeedsResize(const bool bInValue);
+	void RequestCefResize();
 	
 private:
 
@@ -782,9 +786,6 @@ private:
 	/** Used to detect when the widget is hidden*/
 	bool bTickedLastFrame;
 	
-	/** Tracks whether the widget has been resized and needs to be refreshed */
-	bool bNeedsResize;
-
 	/** Tracks whether or not the user initiated a window drag by clicking on a page's drag region. */
 	bool bDraggingWindow;
 
@@ -842,8 +843,11 @@ private:
 	bool bInDirectHwndMode;
 #endif
 
-	bool bHasCorrectNativeCefBuffer = true;
+	bool bHasCorrectNativeCefBuffer;
 	TSharedPtr<FCapturedCefBuffer> CapturedCefBuffer;
+
+	bool bViewportResizePending;
+	size_t NumFramesSinceViewportResize;
 };
 
 typedef FCEFWebBrowserWindow FWebBrowserWindow;
diff --git a/Engine/Source/Runtime/WebBrowser/Private/CEF/CEFWebBrowserWindowRHIHelper.cpp b/Engine/Source/Runtime/WebBrowser/Private/CEF/CEFWebBrowserWindowRHIHelper.cpp
--- a/Engine/Source/Runtime/WebBrowser/Private/CEF/CEFWebBrowserWindowRHIHelper.cpp
+++ b/Engine/Source/Runtime/WebBrowser/Private/CEF/CEFWebBrowserWindowRHIHelper.cpp
@@ -136,7 +136,7 @@ uint64_t FCEFWebBrowserWindowRHIHelper::GetRHIAdapterLuid()
 	return AdapterLuid;
 }
 
-FSlateUpdatableTexture* FCEFWebBrowserWindowRHIHelper::CreateSlateUpdatableTexture(const FUintPoint& TextureSize)
+FSlateUpdatableTexture* FCEFWebBrowserWindowRHIHelper::CreateSlateUpdatableTexture(const FIntPoint& TextureSize)
 {
 #if WITH_ENGINE
 #if PLATFORM_WINDOWS
@@ -162,12 +162,63 @@ void FCEFWebBrowserWindowRHIHelper::ReleaseSlateUpdatableTexture(FSlateUpdatable
 	}
 
 	SlateTextureHandles.Remove(SlateTexture);
+	SlateTexture->Cleanup(); // calls BeginReleaseResource and handles deletion at the end
+#endif // PLATFORM_WINDOWS
+#endif // WITH_ENGINE
+}
+
+FSlateUpdatableTexture* FCEFWebBrowserWindowRHIHelper::ResizeSlateUpdatableTexture(FSlateUpdatableTexture* SlateTexture, const FIntPoint& TextureSize)
+{
+#if WITH_ENGINE
+#if PLATFORM_WINDOWS
+	if (!SlateTexture)
+	{
+		UE_LOG(LogWebBrowser, Error, TEXT("FCEFWebBrowserWindowRHIHelper::ResizeSlateUpdatableTexture() - invalid SlateTexture!"));
+		return nullptr;
+	}
+
+	if (TextureSize.X <= 0 || TextureSize.Y <= 0)
+	{
+		UE_LOG(LogWebBrowser, Error, TEXT("FCEFWebBrowserWindowRHIHelper::ResizeSlateUpdatableTexture() - invalid SlateTexture!"));
+		return nullptr;
+	}
+
+	if (!EnsureShareable(SlateTexture))
+	{
+		UE_LOG(LogWebBrowser, Error, TEXT("FCEFWebBrowserWindowRHIHelper::ResizeSlateUpdatableTexture() - SlateTexture not shareable!"));
+		ReleaseSlateUpdatableTexture(SlateTexture);
+		return nullptr;
+	}
+
+	if (SlateTexture->GetSlateResource()->GetWidth() == TextureSize.X && SlateTexture->GetSlateResource()->GetHeight() == TextureSize.Y)
+	{
+		return SlateTexture;
+	}
 
-	BeginReleaseResource(static_cast<FSlateTexture2DRHIRef*>(SlateTexture));
+	FSlateUpdatableTexture* NewSlateTexture = CreateSlateUpdatableTexture(TextureSize);
+	// Needed so we can copy the existing contents in the new texture now
 	FlushRenderingCommands();
 
-	delete SlateTexture;
+	FIntRect DirtyRect = {
+		0,
+		0,
+		std::min(static_cast<int>(SlateTexture->GetSlateResource()->GetWidth()), TextureSize.X),
+		std::min(static_cast<int>(SlateTexture->GetSlateResource()->GetHeight()), TextureSize.Y)
+	};
+	
+	if (!CopySharedTextureSync(NewSlateTexture, (HANDLE)SlateTextureHandles.FindRef(SlateTexture).Get(), DirtyRect))
+	{
+		UE_LOG(LogWebBrowser, Error, TEXT("FCEFWebBrowserWindowRHIHelper::ResizeSlateUpdatableTexture() - Failed to copy SlateTexture into NewSlateTexture!"));
+	}
+
+	ReleaseSlateUpdatableTexture(SlateTexture);
+
+	return NewSlateTexture;
+#else
+	return nullptr;
 #endif // PLATFORM_WINDOWS
+#else
+	return nullptr;
 #endif // WITH_ENGINE
 }
 
@@ -244,8 +295,6 @@ bool FCEFWebBrowserWindowRHIHelper::EnsureShareable(FSlateUpdatableTexture* Slat
 bool FCEFWebBrowserWindowRHIHelper::CopySharedTextureSync(FSlateUpdatableTexture* SlateTexture, void* SharedHandle, const FIntRect& DirtyIn)
 {
 #if WITH_ENGINE
-	FIntRect Dirty = DirtyIn;
-
 	if (!EnsureShareable(SlateTexture))
 	{
 		return false;
@@ -267,22 +316,12 @@ bool FCEFWebBrowserWindowRHIHelper::CopySharedTextureSync(FSlateUpdatableTexture
 	TComPtr<ID3D11Texture2D> SourceTexture;
 	if (HRESULT Hr = D3D11Device1->OpenSharedResource1((HANDLE)SharedHandle, IID_PPV_ARGS(&SourceTexture)); FAILED(Hr))
 	{
-		UE_LOG(LogWebBrowser, Error, TEXT("FCEFWebBrowserWindowRHIHelper::CopySharedTextureSync() - - ID3D11Device1::OpenSharedResource1 0x%x"), Hr);
-		return false;
-	}
-
-	D3D11_TEXTURE2D_DESC SourceTextureDesc = {};
-	SourceTexture->GetDesc(&SourceTextureDesc);
-	if (SlateTexture->GetSlateResource()->GetWidth() != SourceTextureDesc.Width || SlateTexture->GetSlateResource()->GetHeight() != SourceTextureDesc.Height)
-	{
-		SlateTexture->ResizeTexture(SourceTextureDesc.Width, SourceTextureDesc.Height);
-		FlushRenderingCommands();
-
-		// The native resource has changed, we need to update the shared handle
-		SlateTextureHandles.Remove(SlateTexture);
-		EnsureShareable(SlateTexture);
-
-		Dirty = FIntRect();
+		// Try to open as a D3D11 shared handle instead
+		if (HRESULT Hr2 = D3D11Device1->OpenSharedResource((HANDLE)SharedHandle, IID_PPV_ARGS(&SourceTexture)); FAILED(Hr2))
+		{
+			UE_LOG(LogWebBrowser, Error, TEXT("FCEFWebBrowserWindowRHIHelper::CopySharedTextureSync() - - ID3D11Device1::OpenSharedResource[1] 0x%x [0x%x]"), Hr2, Hr);
+			return false;
+		}
 	}
 
 	TComPtr<ID3D11Texture2D> DestTexture;
@@ -304,12 +343,33 @@ bool FCEFWebBrowserWindowRHIHelper::CopySharedTextureSync(FSlateUpdatableTexture
 		}
 	}
 
-	// Copy the textures on the GPU
-	// Note that as of CEF v128 the Dirty rect sent by CEF always matches the whole SourceTexture no matter what part has actually changed,
-	// so we can just do full copies. If that were to change in the future we can start using CopySubresourceRegion instead
-	D3D11DeviceContext->CopyResource(DestTexture.Get(), SourceTexture.Get());
+	// Copy the texture's dirty subregion on the GPU, limited to DestTexture's size
+	D3D11_TEXTURE2D_DESC SourceTextureDesc = {};
+	SourceTexture->GetDesc(&SourceTextureDesc);
+	D3D11_TEXTURE2D_DESC DestTextureDesc = {};
+	DestTexture->GetDesc(&DestTextureDesc);
+
+	D3D11_BOX SourceDirtyBox = {};
+	SourceDirtyBox.left = std::max(0, DirtyIn.Min.X);
+	SourceDirtyBox.right = std::min({ static_cast<int>(SourceTextureDesc.Width), static_cast<int>(DestTextureDesc.Width), std::max(0, DirtyIn.Max.X) });
+	SourceDirtyBox.top = std::max(0, DirtyIn.Min.Y);
+	SourceDirtyBox.bottom = std::min({ static_cast<int>(SourceTextureDesc.Height), static_cast<int>(DestTextureDesc.Height), std::max(0, DirtyIn.Max.Y) });
+	SourceDirtyBox.front = 0;
+	SourceDirtyBox.back = 1;
+
+	if (SourceDirtyBox.right <= SourceDirtyBox.left || SourceDirtyBox.bottom <= SourceDirtyBox.top)
+	{
+		return false;
+	}
+
+	D3D11DeviceContext->CopySubresourceRegion(
+		DestTexture.Get(), 0,				// pDstResource, DstSubresource
+		SourceDirtyBox.left, SourceDirtyBox.top, 0,	// DstX, DstY, DstZ
+		SourceTexture.Get(), 0,				// pSrcResource, SrcSubresource
+		&SourceDirtyBox						// pSrcBox
+	);
 
-	// The CopyResource call is asynchronous by default, so we need to force a sync flush to ensure it completes before we exit this method
+	// The CopySubresourceRegion call is asynchronous by default, so we need to force a sync flush to ensure it completes before we exit this method
 	D3D11DeviceContext->Flush();
 
 	// Now wait on the GPU to complete the copy and flush
@@ -342,7 +402,7 @@ bool FCEFWebBrowserWindowRHIHelper::CopySharedTextureSync(FSlateUpdatableTexture
 		// Idle wait before retrying
 		if (!SwitchToThread())
 		{
-			Sleep(1); // 1ms
+			Sleep(0);
 		}
 	}
 
diff --git a/Engine/Source/Runtime/WebBrowser/Private/CEF/CEFWebBrowserWindowRHIHelper.h b/Engine/Source/Runtime/WebBrowser/Private/CEF/CEFWebBrowserWindowRHIHelper.h
--- a/Engine/Source/Runtime/WebBrowser/Private/CEF/CEFWebBrowserWindowRHIHelper.h
+++ b/Engine/Source/Runtime/WebBrowser/Private/CEF/CEFWebBrowserWindowRHIHelper.h
@@ -32,8 +32,9 @@ public:
 	static uint64_t GetRHIAdapterLuid();
 
 public:
-	FSlateUpdatableTexture* CreateSlateUpdatableTexture(const FUintPoint& TextureSize);
+	FSlateUpdatableTexture* CreateSlateUpdatableTexture(const FIntPoint& TextureSize);
 	void ReleaseSlateUpdatableTexture(FSlateUpdatableTexture* SlateTexture);
+	FSlateUpdatableTexture* ResizeSlateUpdatableTexture(FSlateUpdatableTexture* SlateTexture, const FIntPoint& TextureSize);
 
 	bool CopySharedTextureSync(FSlateUpdatableTexture* SlateTexture, void* SharedHandle, const FIntRect& DirtyIn);
 
diff --git a/Engine/Source/Runtime/WebBrowser/Private/SWebBrowser.cpp b/Engine/Source/Runtime/WebBrowser/Private/SWebBrowser.cpp
--- a/Engine/Source/Runtime/WebBrowser/Private/SWebBrowser.cpp
+++ b/Engine/Source/Runtime/WebBrowser/Private/SWebBrowser.cpp
@@ -7,7 +7,6 @@
 #include "Widgets/Images/SThrobber.h"
 #include "Framework/Application/SlateApplication.h"
 #if WITH_CEF3
-#include "Framework/Docking/TabManager.h"
 #include "Async/Async.h"
 #include "CEF/CEFWebBrowserWindow.h"
 #endif // WITH_CEF3
@@ -56,46 +55,9 @@ static FWebBrowserSingleton& GetWebBrowserSingletonChecked()
 	if (WebBrowserWindow.IsValid()) 
 	{
 		if (const TSharedRef<FCEFWebBrowserWindow> CefWebBrowserWindow = GetCefWebBrowserWindowChecked(WebBrowserWindow);
-			CefWebBrowserWindow->HasResizeBug())
+			CefWebBrowserWindow->HasResizeBug() && CefWebBrowserWindow->IsCefResizeNeeded())
 		{
-			// NOTE - We won't end up here if the CEF version doesn't require it.
-			
-			const bool bIsFloatingWindow = (FSlateApplication::Get().FindWidgetWindow(AsShared()) != FGlobalTabmanager::Get()->GetRootWindow());
-			
-
-			// * This is necessary during active widget resizes, whether or not the widget is in a floating window. *
-			// If web browser doesn't have the correct CEF buffer, tell it to resize.
-			// NOTE - Call this before the 'forced tick' code below.
-			
-			if (!CefWebBrowserWindow->HasCorrectNativeCefBuffer())
-			{
-				if (bIsFloatingWindow)
-				{
-					// Seems necessary to make web browser content appear, if it's already in a floating window when the Editor starts up.
-					// But this is not good to call if docked in editor!
-					CefWebBrowserWindow->SetViewportSize(BrowserView->GetCachedGeometry().Size.IntPoint(), BrowserView->GetCachedGeometry().AbsolutePosition.IntPoint());
-				}
-
-				// Always seems necessary (even if viewport size was just changed above.)
-				CefWebBrowserWindow->SetNeedsResize(true);
-			}
-	
-#if PLATFORM_WINDOWS
-			// * This is necessary during active widget resizes, only if the widget is in a floating window, and only when on Windows. *
-			// When we are (1) running on Windows (2) not docked in the Editor (floating), and (3) our window is actively being resized - then the FTS
-			// Ticker interface on the Web Browser Singleton will be suspended. But our Tick() function here will not be suspended. So under these and
-			// certain other conditions (see below), we want to force Web Browser Singleton to tick, and we can do that from here. But, we also can't tick
-			// it too fast or too often, or we'll overwhelm the CEF message pump, and it will start ignoring what we want it to do. We also can't tick it
-			// too slow, or it can ignore the tick.
-			// TODO - In this floating window case only, the web browser can still end up stuck with the wrong texture, when mouse stops moving during
-			//		  an active resize, until mouse is released. Why?
-			
-			if (FWebBrowserSingleton& WebBrowserSingleton = GetWebBrowserSingletonChecked();
-				bIsFloatingWindow && FPlatformTime::Seconds() - WebBrowserSingleton.GetPreviousTickTimeSeconds() >= WebBrowserSingleton.GetExternalTickCallWaitSeconds())
-			{
-				WebBrowserSingleton.Tick(InDeltaTime);
-			}
-#endif // PLATFORM_WINDOWS
+			CefWebBrowserWindow->RequestCefResize();
 		}
 	}
 #endif // WITH_CEF3
diff --git a/Engine/Source/Runtime/WebBrowser/Private/WebBrowserSingleton.cpp b/Engine/Source/Runtime/WebBrowser/Private/WebBrowserSingleton.cpp
--- a/Engine/Source/Runtime/WebBrowser/Private/WebBrowserSingleton.cpp
+++ b/Engine/Source/Runtime/WebBrowser/Private/WebBrowserSingleton.cpp
@@ -836,8 +836,6 @@ bool FWebBrowserSingleton::Tick(float DeltaTime)
 {
 	QUICK_SCOPE_CYCLE_COUNTER(STAT_FWebBrowserSingleton_Tick);
 
-	PreviousTickTimeSeconds = FPlatformTime::Seconds();
-	
 #if WITH_CEF3
 	if (bAllowCEF)
 	{
@@ -887,12 +885,6 @@ bool FWebBrowserSingleton::Tick(float DeltaTime)
 		float MinMessageLoopSeconds = 1.0f / MinMessageLoopHz;
 		float MaxForcedMessageLoopSeconds = 1.0f / MaxForcedMessageLoopHz;
 
-		// Set the external tick-call wait time to be in the acceptable range, but leaning heavily toward the longer time. That way, the ticker in
-		// this class is likely to tick itself first, whenever possible. If the ticker is suspended (see where this is used outside of this class),
-		// then this should be an acceptable wait time.
-		static constexpr float ExternalTickCallWaitSecondsScale = 0.95f; 
-		ExternalTickCallWaitSeconds = ExternalTickCallWaitSecondsScale * MaxForcedMessageLoopSeconds;
-		
 		static float SecondsSinceLastPump = 0;
 		static float SecondsSinceLastAppFocusCheck = MaxForcedMessageLoopSeconds;
 		static float SecondsToNextForcedPump = MaxForcedMessageLoopSeconds;
@@ -982,21 +974,6 @@ bool FWebBrowserSingleton::Tick(float DeltaTime)
 	return true;
 }
 
-#if WITH_CEF3
-double FWebBrowserSingleton::GetPreviousTickTimeSeconds() const
-{
-	return PreviousTickTimeSeconds;
-}
-#endif
-
-#if WITH_CEF3
-float FWebBrowserSingleton::GetExternalTickCallWaitSeconds() const
-{
-	static constexpr float DefaultExternalTickCallWaitSeconds = 0.064f;
-	return (ExternalTickCallWaitSeconds <= 0.0f ? DefaultExternalTickCallWaitSeconds : ExternalTickCallWaitSeconds);
-}
-#endif
-
 FString FWebBrowserSingleton::GetCurrentLocaleCode()
 {
 	FCultureRef Culture = FInternationalization::Get().GetCurrentCulture();
diff --git a/Engine/Source/Runtime/WebBrowser/Private/WebBrowserSingleton.h b/Engine/Source/Runtime/WebBrowser/Private/WebBrowserSingleton.h
--- a/Engine/Source/Runtime/WebBrowser/Private/WebBrowserSingleton.h
+++ b/Engine/Source/Runtime/WebBrowser/Private/WebBrowserSingleton.h
@@ -164,12 +164,6 @@ public:
 	virtual bool Tick(float DeltaTime) override;
 
 #if WITH_CEF3
-	/** Returns the time in seconds that an external caller should wait after previous tick, before externally calling Tick(). */
-	WEBBROWSER_API float GetExternalTickCallWaitSeconds() const;
-
-	/** Returns the time of the most recent call to Tick(). */
-	WEBBROWSER_API double GetPreviousTickTimeSeconds() const;
-
 	/** Return true if this URL will support adding an Authorization header to it */
 	bool URLRequestAllowsCredentials(const FString& URL);
 #endif
@@ -227,12 +221,6 @@ private:
 	/** Reference to UWebBrowser's translucent material*/
 	TSoftObjectPtr<UMaterialInterface> DefaultTranslucentMaterial;
 #endif
-
-	/* Time when our Tick() was last called. */
-	double PreviousTickTimeSeconds = 0.0;
-	
-	/** How long an external caller of our Tick() function should wait, after the previous tick time, before calling our Tick(). */ 
-	float ExternalTickCallWaitSeconds = 0.0;
 };
 
 PRAGMA_ENABLE_DEPRECATION_WARNINGS
