WebBrowser: improve performance during non-accelerated resizing with CEF

This is the no-gpu-acceleration counterpart to CL 50668052, which is still used on Linux and Mac.
Reworked the code to match the resizing logic in the accelerated case:
- removed the SWebBrowser::OnPaint overload that substituted normal painting with a custom texture
- instead we use the existing UpdatableTexture, and copy the captured cef buffer into it when resized to avoid stretching effects
- renamed FCapturedCefBuffer to FCEFCapturedFrame and moved it to the correct CEF subfolder
- adjusted the frame size checks to avoid requesting CEF resizes too quickly

Also fixed a minor bug on the accelerated resizing path, to do viewport size checks for main view updates only (not for popups).

#rnx

#Synced-CL 50700524

diff --git a/Engine/Source/Runtime/WebBrowser/Private/CEF/CEFCapturedFrame.cpp b/Engine/Source/Runtime/WebBrowser/Private/CEF/CEFCapturedFrame.cpp
new file mode 100644
--- /dev/null
+++ b/Engine/Source/Runtime/WebBrowser/Private/CEF/CEFCapturedFrame.cpp
@@ -0,0 +1,80 @@
+// Copyright Epic Games, Inc. All Rights Reserved.
+
+
+#include "CEFCapturedFrame.h"
+
+#if WITH_CEF3
+
+#include "Textures/SlateTextureData.h"
+
+static constexpr uint32 BUFFER_ELEMENT_SIZE = 4; /*PF_B8G8R8A8*/
+
+bool FCEFCapturedFrame::SetBufferAsB8G8R8A8(const void* InBufferB8G8R8A8, const int32 InBufferWidth, const int32 InBufferHeight, const bool bSkipTransparencyTest)
+{
+	if (!InBufferB8G8R8A8 || InBufferWidth <= 0 || InBufferHeight <= 0)
+	{
+		ClearBuffer();
+		return false;
+	}
+
+	// When resizing, CEF can sometimes send us buffers with the contents stretched/compressed inside the frame and surrounded by transparent bars.
+	// We detect this by sampling the alpha channel of the first pixel, in the top-left corner. If it holds a 0x0 then the frame is bad and should be skipped.
+	if (!bSkipTransparencyTest && (static_cast<const uint8*>(InBufferB8G8R8A8)[3] == 0x0))
+	{
+		return false;
+	}
+
+	// Don't capture anymore if we already have a buffer of the same size
+	if (InBufferWidth != BufferDimensions.X || InBufferHeight != BufferDimensions.Y)
+	{
+		BufferDimensions = FUintPoint(InBufferWidth, InBufferHeight);
+
+		const TArray<uint8>::SizeType NumBytes = BufferDimensions.X * BufferDimensions.Y * BUFFER_ELEMENT_SIZE;
+		if (BufferData.Num() != NumBytes)
+		{
+			BufferData.SetNum(NumBytes);
+		}
+		FMemory::Memcpy(BufferData.GetData(), InBufferB8G8R8A8, NumBytes);
+	}
+
+	return true;
+}
+
+bool FCEFCapturedFrame::ClearBuffer()
+{
+	BufferData.Empty();
+	BufferDimensions = FUintPoint::ZeroValue;
+	return true;
+}
+
+bool FCEFCapturedFrame::IsEmpty() const
+{
+	return BufferData.IsEmpty();
+}
+
+FSlateTextureData* FCEFCapturedFrame::CreateSlateTextureData(const FIntPoint& TextureSize) const
+{
+	if (BufferData.IsEmpty() || TextureSize.X <= 0 || TextureSize.Y <= 0)
+	{
+		return new FSlateTextureData(nullptr, 0, 0, BUFFER_ELEMENT_SIZE);
+	}
+
+	const FUintPoint TextureSizeChecked = (FUintPoint)TextureSize;
+	if (BufferDimensions == TextureSizeChecked)
+	{
+		return new FSlateTextureData(BufferData.GetData(), BufferDimensions.X, BufferDimensions.Y, BUFFER_ELEMENT_SIZE);
+	}
+
+	TArray<uint8> TextureDataBuffer;
+	TextureDataBuffer.SetNum(TextureSizeChecked.X * TextureSizeChecked.Y * BUFFER_ELEMENT_SIZE);
+	for (uint32 Row = 0, MaxRow = std::min(BufferDimensions.Y, TextureSizeChecked.Y); Row < MaxRow; ++Row)
+	{
+		const size_t SourceOffset = Row * BufferDimensions.X * BUFFER_ELEMENT_SIZE;
+		const size_t TargetOffset = Row * TextureSizeChecked.X * BUFFER_ELEMENT_SIZE;
+		FMemory::Memcpy(TextureDataBuffer.GetData() + TargetOffset, BufferData.GetData() + SourceOffset, std::min(BufferDimensions.X, TextureSizeChecked.X) * BUFFER_ELEMENT_SIZE);
+	}
+
+	return new FSlateTextureData(TextureSizeChecked.X, TextureSizeChecked.Y, BUFFER_ELEMENT_SIZE, std::move(TextureDataBuffer));
+}
+
+#endif // WITH_CEF3
diff --git a/Engine/Source/Runtime/WebBrowser/Private/CEF/CEFCapturedFrame.h b/Engine/Source/Runtime/WebBrowser/Private/CEF/CEFCapturedFrame.h
new file mode 100644
--- /dev/null
+++ b/Engine/Source/Runtime/WebBrowser/Private/CEF/CEFCapturedFrame.h
@@ -0,0 +1,31 @@
+// Copyright Epic Games, Inc. All Rights Reserved.
+
+#pragma once
+
+#include "CoreMinimal.h"
+
+#include "CEFCapturedFrame.generated.h"
+
+struct FSlateTextureData;
+
+USTRUCT()
+struct FCEFCapturedFrame
+{
+	GENERATED_BODY()
+
+#if WITH_CEF3
+public:
+
+	bool SetBufferAsB8G8R8A8(const void* InBufferB8G8R8A8, const int32 InBufferWidth, const int32 InBufferHeight, const bool bSkipTransparencyTest);
+	bool ClearBuffer();
+	
+	bool IsEmpty() const;
+
+	FSlateTextureData* CreateSlateTextureData(const FIntPoint& TextureSize) const;
+
+private:
+
+	TArray<uint8> BufferData;
+	FUintPoint BufferDimensions = FUintPoint::ZeroValue;
+#endif // WITH_CEF3
+};
diff --git a/Engine/Source/Runtime/WebBrowser/Private/CEF/CEFWebBrowserWindow.cpp b/Engine/Source/Runtime/WebBrowser/Private/CEF/CEFWebBrowserWindow.cpp
--- a/Engine/Source/Runtime/WebBrowser/Private/CEF/CEFWebBrowserWindow.cpp
+++ b/Engine/Source/Runtime/WebBrowser/Private/CEF/CEFWebBrowserWindow.cpp
@@ -22,6 +22,7 @@
 #include "Async/Async.h"
 
 #if PLATFORM_MAC
+#include "Mac/MacSystemIncludes.h"
 // Needed for character code definitions
 #include <Carbon/Carbon.h>
 #include <AppKit/NSEvent.h>
@@ -427,7 +428,6 @@ FCEFWebBrowserWindow::FCEFWebBrowserWindow(CefRefPtr<CefBrowser> InBrowser, CefR
 	, bInDirectHwndMode(false)
 #endif
 	, bHasCorrectNativeCefBuffer(true)
-	, bViewportResizePending(false)
 	, NumFramesSinceViewportResize(0)
 {
 	check(InBrowser.get() != nullptr);
@@ -644,7 +644,7 @@ void FCEFWebBrowserWindow::SetViewportSize(FIntPoint WindowSize, FIntPoint Windo
 		ViewportDPIScaleFactor = WindowDPIScaleFactor;
 		FIntPoint WindowSizeScaled = (FVector2D(WindowSize) * WindowDPIScaleFactor).IntPoint();
 
-		if (IsValid())
+		if (IsValid() && WindowSizeScaled.X > 0 && WindowSizeScaled.Y > 0)
 		{
 #if PLATFORM_WINDOWS
 			HWND NativeHandle = InternalCefBrowser->GetHost()->GetWindowHandle();
@@ -678,14 +678,40 @@ void FCEFWebBrowserWindow::SetViewportSize(FIntPoint WindowSize, FIntPoint Windo
 				}
 			}
 #endif
-			bHasCorrectNativeCefBuffer = false;
-			NumFramesSinceViewportResize = 0;
 			RequestCefResize();
 
+			// The lines below create a new window texture with the new size, also copying the old texture contents 1:1 into the new one, so they don't look stretched.
+			// But, in the non-accelerated case we only want to use CapturedCefFrame's contents if they're from a recent frame (ie. an ongoing drag-resize)
 			if (RHIRenderHelper)
 			{
 				UpdatableTextures[PET_VIEW] = RHIRenderHelper->ResizeSlateUpdatableTexture(UpdatableTextures[PET_VIEW], WindowSizeScaled);
 			}
+			else if (CapturedCefFrame.IsValid() && !CapturedCefFrame->IsEmpty() && NumFramesSinceViewportResize < 10)
+			{
+				FSlateRenderer* Renderer = GetRenderer();
+				if (!Renderer)
+				{
+					return;
+				}
+
+				FSlateUpdatableTexture* ResizedTexture = Renderer->CreateUpdatableTexture(WindowSizeScaled.X, WindowSizeScaled.Y);
+				HandleRenderingError();
+				if (!ResizedTexture)
+				{
+					return;
+				}
+
+				// will safely delete the FSlateTextureData pointer when done with it on the rendering thread
+				ResizedTexture->UpdateTextureThreadSafeWithTextureData(CapturedCefFrame->CreateSlateTextureData(WindowSizeScaled));
+				HandleRenderingError();
+				
+				Renderer->ReleaseUpdatableTexture(UpdatableTextures[PET_VIEW]);
+				HandleRenderingError();
+
+				UpdatableTextures[PET_VIEW] = ResizedTexture;
+			}
+
+			NumFramesSinceViewportResize = 0;
 		}
 	}
 }
@@ -1970,26 +1996,49 @@ void FCEFWebBrowserWindow::HandleRenderingError()
 
 void FCEFWebBrowserWindow::OnPaint(CefRenderHandler::PaintElementType Type, const CefRenderHandler::RectList& DirtyRects, const void* Buffer, int Width, int Height)
 {
-	if (ViewportSize == FIntPoint::ZeroValue)
+	if (!Buffer || Width <= 0 || Height <= 0)
 	{
 		return;
 	}
 
+	const FIntRect DirtyRect = (DirtyRects.size() == 1) ? FIntRect(DirtyRects[0].x, DirtyRects[0].y, DirtyRects[0].x + DirtyRects[0].width, DirtyRects[0].y + DirtyRects[0].height) : FIntRect();
+	const FIntRect ViewportRect = FIntRect(0, 0, ViewportSize.X, ViewportSize.Y).Scale(ViewportDPIScaleFactor);
+	if (DirtyRect.IsEmpty() || ViewportRect.IsEmpty())
+	{
+		return;
+	}
+
+	++NumFramesSinceViewportResize;
 
 #if UE_CEF_HAS_RESIZE_BUG
-	const FIntPoint PaintedBufferSize(Width / ViewportDPIScaleFactor, Height / ViewportDPIScaleFactor);
-	bHasCorrectNativeCefBuffer = (PaintedBufferSize == ViewportSize);
-	if (CapturedCefBuffer.IsValid())
+	if (Type == PET_VIEW)
 	{
-		// We attempt to capture a good CEF buffer, especially when it is correct. The CEF buffer size does not actually need to match the viewport
-		// size to end up being captured, but also may not be captured if it fails another test (which we skip if we know we have the correct CEF
-		// buffer.) See code.
+		const bool bWasCorrectNativeCefBuffer = bHasCorrectNativeCefBuffer;
+		bHasCorrectNativeCefBuffer = (FIntPoint(Width, Height) == ViewportRect.Size());
+
+		if (!bHasCorrectNativeCefBuffer)
+		{
+			UE_LOG(LogWebBrowser, Verbose, TEXT("FCEFWebBrowserWindow::OnPaint: Frame size still wrong after %zu frames: %dx%d -> %dx%d (%d,%d %dx%d)"), NumFramesSinceViewportResize, ViewportRect.Width(), ViewportRect.Height(), Width, Height, DirtyRect.Min.X, DirtyRect.Min.Y, DirtyRect.Width(), DirtyRect.Height());
+		}
+		else if (!bWasCorrectNativeCefBuffer)
+		{
+			UE_LOG(LogWebBrowser, Verbose, TEXT("FCEFWebBrowserWindow::OnPaint: Frame size OK again after %zu frames: %dx%d -> %dx%d (%d,%d %dx%d)"), NumFramesSinceViewportResize, ViewportRect.Width(), ViewportRect.Height(), Width, Height, DirtyRect.Min.X, DirtyRect.Min.Y, DirtyRect.Width(), DirtyRect.Height());
+		}
 
-		CapturedCefBuffer->SetBufferAsB8G8R8A8(Buffer, Width, Height, ViewportDPIScaleFactor, bHasCorrectNativeCefBuffer);
+		if (CapturedCefFrame.IsValid())
+		{
+			// Capture the frame to be used during resizing (even if it's not the right size)
+			if (!CapturedCefFrame->SetBufferAsB8G8R8A8(Buffer, Width, Height, bHasCorrectNativeCefBuffer))
+			{
+				UE_LOG(LogWebBrowser, Verbose, TEXT("FCEFWebBrowserWindow::OnPaint: Detected corrupted CEF frame!"));
+				// force a CEF resize even if we were below the threshold, as CEF considers those "good" frames and will stop sending more until requested
+				NumFramesSinceViewportResize = MIN_FRAMES_FOR_CEF_RESIZE;
+				return;
+			}
+		}
 	}
-#endif
+#endif // UE_CEF_HAS_RESIZE_BUG
 
-	
 	bool bNeedsRedraw = false;
 	if (bUsingAcceleratedPaint)
 	{
@@ -2016,32 +2065,25 @@ void FCEFWebBrowserWindow::OnPaint(CefRenderHandler::PaintElementType Type, cons
 
 	if (UpdatableTextures[Type] != nullptr)
 	{
-		// Note that with more recent versions of CEF, the DirtyRects will always contain a single element, as it merges all dirty areas into a single rectangle before calling OnPaint
-		// In case that should change in the future, we'll simply update the entire area if DirtyRects is not a single element.
-		FIntRect Dirty = (DirtyRects.size() == 1) ? FIntRect(DirtyRects[0].x, DirtyRects[0].y, DirtyRects[0].x + DirtyRects[0].width, DirtyRects[0].y + DirtyRects[0].height) : FIntRect();
-
 		if (Type == PET_VIEW && BufferedVideo.IsValid() )
 		{
 			// If we're using bufferedVideo, submit the frame to it
-			bNeedsRedraw = BufferedVideo->SubmitFrame(Width, Height, Buffer, Dirty);
+			bNeedsRedraw = BufferedVideo->SubmitFrame(Width, Height, Buffer, DirtyRect);
 		}
 		else
 		{
-			// Ignore the frame and force a CEF repaint if the texture size doesn't match the current viewport size.
-			// This is another side-effect of [REDACTED_PROFILE_URL]/cef/issues/3826:
-			// When resizing the view, the last frame we receive occasionally still has a previous viewport size,
-			// which results in contents looking stretched or with black borders.
-#if UE_CEF_HAS_RESIZE_BUG
-			if (!CapturedCefBuffer.IsValid() && Type == PET_VIEW && !bHasCorrectNativeCefBuffer)
+			if (Type == PET_VIEW && !bHasCorrectNativeCefBuffer)
 			{
-				// Note that calling Invalidate() doesn't work here, and there's no other CEF API to trigger a repaint.
-				InternalCefBrowser->GetHost()->WasHidden(true);
-				InternalCefBrowser->GetHost()->WasHidden(false);
+				if (CapturedCefFrame.IsValid() && !CapturedCefFrame->IsEmpty())
+				{
+					// will safely delete the FSlateTextureData pointer when done with it on the rendering thread
+					UpdatableTextures[Type]->UpdateTextureThreadSafeWithTextureData(CapturedCefFrame->CreateSlateTextureData(ViewportRect.Size()));
+					HandleRenderingError();
+				}
 			}
 			else
-#endif
 			{
-				UpdatableTextures[Type]->UpdateTextureThreadSafeRaw(Width, Height, Buffer, Dirty);
+				UpdatableTextures[Type]->UpdateTextureThreadSafeRaw(Width, Height, Buffer, DirtyRect);
 				HandleRenderingError();
 				bNeedsRedraw = true;
 			}
@@ -2074,6 +2116,15 @@ void FCEFWebBrowserWindow::OnAcceleratedPaint(CefRenderHandler::PaintElementType
 	const CefAcceleratedPaintInfo& Info)
 #endif
 {
+	FIntRect DirtyRect = (DirtyRects.size() == 1) ? FIntRect(DirtyRects[0].x, DirtyRects[0].y, DirtyRects[0].x + DirtyRects[0].width, DirtyRects[0].y + DirtyRects[0].height) : FIntRect();
+	FIntRect ViewportRect = FIntRect(0, 0, ViewportSize.X, ViewportSize.Y).Scale(ViewportDPIScaleFactor);
+	if (DirtyRect.IsEmpty() || ViewportRect.IsEmpty())
+	{
+		return;
+	}
+
+	++NumFramesSinceViewportResize;
+
 	bool bNeedsRedraw = false;
 	if (!bUsingAcceleratedPaint)
 	{
@@ -2090,13 +2141,6 @@ void FCEFWebBrowserWindow::OnAcceleratedPaint(CefRenderHandler::PaintElementType
 		}
 	}
 
-	FIntRect DirtyRect = (DirtyRects.size() == 1) ? FIntRect(DirtyRects[0].x, DirtyRects[0].y, DirtyRects[0].x + DirtyRects[0].width, DirtyRects[0].y + DirtyRects[0].height) : FIntRect();
-	FIntRect ViewportRect = FIntRect(0, 0, ViewportSize.X, ViewportSize.Y).Scale(ViewportDPIScaleFactor);
-	if (DirtyRect.IsEmpty() || ViewportRect.IsEmpty())
-	{
-		return;
-	}
-
 #if CEF_VERSION_MAJOR >= 128
 #  if PLATFORM_WINDOWS
 	void* SharedHandle = Info.shared_texture_handle;
@@ -2115,27 +2159,30 @@ void FCEFWebBrowserWindow::OnAcceleratedPaint(CefRenderHandler::PaintElementType
 	LastPaintedSharedHandle = SharedHandle;
 #endif
 
-	bool bWasCorrectNativeCefBuffer = bHasCorrectNativeCefBuffer;
-	bHasCorrectNativeCefBuffer = (DirtyRect == ViewportRect);
-	bViewportResizePending = false;
-	++NumFramesSinceViewportResize;
-	if (!bHasCorrectNativeCefBuffer)
-	{
-		UE_LOG(LogWebBrowser, Verbose, TEXT("FCEFWebBrowserWindow::OnAcceleratedPaint: Frame size still wrong after %zu frames: %dx%d -> %dx%d (%d, %d)"), NumFramesSinceViewportResize, ViewportRect.Width(), ViewportRect.Height(), DirtyRect.Width(), DirtyRect.Height(), DirtyRect.Min.X, DirtyRect.Min.Y);
-	}
-	else if (!bWasCorrectNativeCefBuffer)
+#if UE_CEF_HAS_RESIZE_BUG
+	if (Type == PET_VIEW)
 	{
-		UE_LOG(LogWebBrowser, Verbose, TEXT("FCEFWebBrowserWindow::OnAcceleratedPaint: Frame size OK again after %zu frames: %dx%d -> %dx%d (%d, %d)"), NumFramesSinceViewportResize, ViewportRect.Width(), ViewportRect.Height(), DirtyRect.Width(), DirtyRect.Height(), DirtyRect.Min.X, DirtyRect.Min.Y);
-	}
+		bool bWasCorrectNativeCefBuffer = bHasCorrectNativeCefBuffer;
+		bHasCorrectNativeCefBuffer = (DirtyRect == ViewportRect);
+		if (!bHasCorrectNativeCefBuffer)
+		{
+			UE_LOG(LogWebBrowser, Verbose, TEXT("FCEFWebBrowserWindow::OnAcceleratedPaint: Frame size still wrong after %zu frames: %dx%d -> %dx%d (%d, %d)"), NumFramesSinceViewportResize, ViewportRect.Width(), ViewportRect.Height(), DirtyRect.Width(), DirtyRect.Height(), DirtyRect.Min.X, DirtyRect.Min.Y);
+		}
+		else if (!bWasCorrectNativeCefBuffer)
+		{
+			UE_LOG(LogWebBrowser, Verbose, TEXT("FCEFWebBrowserWindow::OnAcceleratedPaint: Frame size OK again after %zu frames: %dx%d -> %dx%d (%d, %d)"), NumFramesSinceViewportResize, ViewportRect.Width(), ViewportRect.Height(), DirtyRect.Width(), DirtyRect.Height(), DirtyRect.Min.X, DirtyRect.Min.Y);
+		}
 
-	// During resizing, CEF occasionally sends textures with black bars either above/below or left/right and distorted content in between, we want to ignore these textures.
-	// Fortunately, the DirtyRect for these textures (and only these) has a non-zero X/Y depending on the black bars orientation, so we can detect them without inspecting the content.
-	if (DirtyRect.Min.X > 0 || DirtyRect.Min.Y > 0)
-	{
-		// force a CEF resize even if we were below the threshold, as CEF considers those "good" frames and will stop sending more until requested
-		NumFramesSinceViewportResize = MIN_FRAMES_FOR_CEF_RESIZE;
-		return;
+		// During resizing, CEF occasionally sends textures with black bars either above/below or left/right and distorted content in between, we want to ignore these textures.
+		// Fortunately, the DirtyRect for these textures (and only these) has a non-zero X/Y depending on the black bars orientation, so we can detect them without inspecting the content.
+		if (DirtyRect.Min.X > 0 || DirtyRect.Min.Y > 0)
+		{
+			// force a CEF resize even if we were below the threshold, as CEF considers those "good" frames and will stop sending more until requested
+			NumFramesSinceViewportResize = MIN_FRAMES_FOR_CEF_RESIZE;
+			return;
+		}
 	}
+#endif // UE_CEF_HAS_RESIZE_BUG
 
 	if (RHIRenderHelper && RHIRenderHelper->CopySharedTextureSync(UpdatableTextures[Type], SharedHandle, DirtyRect))
 	{
@@ -2869,31 +2916,20 @@ CefRefPtr<CefDictionaryValue> FCEFWebBrowserWindow::GetProcessInfo()
 	return Retval;
 }
 
-bool FCEFWebBrowserWindow::HasCorrectNativeCefBuffer() const
+void FCEFWebBrowserWindow::SetCaptureCefFrames(const bool bValue)
 {
 #if UE_CEF_HAS_RESIZE_BUG
-	return bHasCorrectNativeCefBuffer;
-#else
-	return true;
-#endif
-}
-
-void FCEFWebBrowserWindow::SetUseCapturedCefBuffer(const bool bValue)
-{
-#if UE_CEF_HAS_RESIZE_BUG && UE_WITH_CAPTURED_CEF_BUFFER
 	if (bValue)
 	{
-		if (!CapturedCefBuffer.IsValid())
+		if (!CapturedCefFrame.IsValid())
 		{
-			CapturedCefBuffer = MakeShared<FCapturedCefBuffer>();
+			CapturedCefFrame = MakeShared<FCEFCapturedFrame>();
 		}
 	}
 	else
 	{
-		CapturedCefBuffer.Reset();
+		CapturedCefFrame.Reset();
 	}
-#else
-	CapturedCefBuffer.Reset();
 #endif
 }
 
@@ -2906,14 +2942,9 @@ bool FCEFWebBrowserWindow::HasResizeBug() const
 #endif
 }
 
-TSharedPtr<FCapturedCefBuffer> FCEFWebBrowserWindow::GetCapturedCefBuffer()
-{
-	return CapturedCefBuffer; // ..may not be set
-}
-
 bool FCEFWebBrowserWindow::IsCefResizeNeeded() const
 {
-	if (!bHasCorrectNativeCefBuffer && !bViewportResizePending && NumFramesSinceViewportResize >= MIN_FRAMES_FOR_CEF_RESIZE)
+	if (!bHasCorrectNativeCefBuffer && NumFramesSinceViewportResize >= MIN_FRAMES_FOR_CEF_RESIZE)
 	{
 		UE_LOG(LogWebBrowser, Verbose, TEXT("FCEFWebBrowserWindow::IsCefResizeNeeded: Forcing CEF resize"));
 		return true;
@@ -2939,8 +2970,6 @@ void FCEFWebBrowserWindow::RequestCefResize()
 		InternalCefBrowser->GetHost()->Invalidate(PET_POPUP);
 	}
 #endif
-
-	bViewportResizePending = true;
 }
 
 bool FCEFWebBrowserWindow::OnProcessMessageReceived(CefRefPtr<CefBrowser> Browser, CefRefPtr<CefFrame> frame, CefProcessId SourceProcess, CefRefPtr<CefProcessMessage> Message)
diff --git a/Engine/Source/Runtime/WebBrowser/Private/CEF/CEFWebBrowserWindow.h b/Engine/Source/Runtime/WebBrowser/Private/CEF/CEFWebBrowserWindow.h
--- a/Engine/Source/Runtime/WebBrowser/Private/CEF/CEFWebBrowserWindow.h
+++ b/Engine/Source/Runtime/WebBrowser/Private/CEF/CEFWebBrowserWindow.h
@@ -16,7 +16,7 @@
 
 #include "CEFLibCefIncludes.h"
 
-#include "CapturedCefBuffer.h"
+#include "CEFCapturedFrame.h"
 
 #endif
 
@@ -585,28 +585,18 @@ public:
 	* Return true if this URL will support adding an Authorization header to it
 	*/
 	bool URLRequestAllowsCredentials(const FString& URL) const { return WebBrowserHandler->URLRequestAllowsCredentials(URL); }
-
-	/**
-	 * Tells us whether this has the correct CEF image buffer.
-	 */
-	bool HasCorrectNativeCefBuffer() const;
 	
 	/**
 	 * Whether we should use a workaround for this bug -
 	 * [REDACTED_PROFILE_URL]/cef/issues/3826
 	 */
-	void SetUseCapturedCefBuffer(const bool bValue);
+	void SetCaptureCefFrames(const bool bValue);
 	
 	/**
 	 * Whether the CEF version we're using has a resize bug that sometimes requires requesting new frames after a resize.
 	 */
 	bool HasResizeBug() const;
 	
-	/**
-	 * Gets an undistorted buffer captured from CEF. Format is PF_B8G8R8A8.
-	 */
-	TSharedPtr<FCapturedCefBuffer> GetCapturedCefBuffer();
-
 	/**
 	 * Indicates a CEF resize is needed, using heuristics to avoid too many resize requests during an active drag-resize.
 	 */
@@ -844,9 +834,8 @@ private:
 #endif
 
 	bool bHasCorrectNativeCefBuffer;
-	TSharedPtr<FCapturedCefBuffer> CapturedCefBuffer;
+	TSharedPtr<FCEFCapturedFrame> CapturedCefFrame;
 
-	bool bViewportResizePending;
 	size_t NumFramesSinceViewportResize;
 };
 
diff --git a/Engine/Source/Runtime/WebBrowser/Private/CapturedCefBuffer.cpp b/Engine/Source/Runtime/WebBrowser/Private/CapturedCefBuffer.cpp
deleted file mode 100644
diff --git a/Engine/Source/Runtime/WebBrowser/Private/CapturedCefBuffer.h b/Engine/Source/Runtime/WebBrowser/Private/CapturedCefBuffer.h
deleted file mode 100644
diff --git a/Engine/Source/Runtime/WebBrowser/Private/SWebBrowser.cpp b/Engine/Source/Runtime/WebBrowser/Private/SWebBrowser.cpp
--- a/Engine/Source/Runtime/WebBrowser/Private/SWebBrowser.cpp
+++ b/Engine/Source/Runtime/WebBrowser/Private/SWebBrowser.cpp
@@ -6,12 +6,10 @@
 #include "Widgets/Input/SButton.h"
 #include "Widgets/Images/SThrobber.h"
 #include "Framework/Application/SlateApplication.h"
+
 #if WITH_CEF3
-#include "Async/Async.h"
 #include "CEF/CEFWebBrowserWindow.h"
 #endif // WITH_CEF3
-#include "WebBrowserModule.h"
-#include "WebBrowserSingleton.h"
 
 #define LOCTEXT_NAMESPACE "WebBrowser"
 
@@ -26,26 +24,6 @@ SWebBrowser::~SWebBrowser()
 #endif
 }
 
-#if WITH_CEF3
-static TSharedRef<FCEFWebBrowserWindow> GetCefWebBrowserWindowChecked(const TSharedPtr<IWebBrowserWindow>& WebBrowserWindow)
-{
-	check(WebBrowserWindow.IsValid());
-	
-	const TSharedPtr<FCEFWebBrowserWindow> CefWebBrowserWindow = StaticCastSharedPtr<FCEFWebBrowserWindow>(WebBrowserWindow);
-	check(CefWebBrowserWindow.IsValid());
-	
-	return CefWebBrowserWindow.ToSharedRef();
-}
-#endif // WITH_CEF3
-
-static FWebBrowserSingleton& GetWebBrowserSingletonChecked()
-{
-	IWebBrowserModule& WebBrowserModule = FModuleManager::LoadModuleChecked<IWebBrowserModule>("WebBrowser");
-	FWebBrowserSingleton* WebBrowserSingleton = static_cast<FWebBrowserSingleton*>(WebBrowserModule.GetSingleton());
-	check(WebBrowserSingleton);
-	return *WebBrowserSingleton;
-}
-
 /*virtual*/ void SWebBrowser::Tick(const FGeometry& AllottedGeometry, const double InCurrentTime, const float InDeltaTime)
 {
 	SCompoundWidget::Tick(AllottedGeometry, InCurrentTime, InDeltaTime);
@@ -54,8 +32,8 @@ static FWebBrowserSingleton& GetWebBrowserSingletonChecked()
 #if WITH_CEF3
 	if (WebBrowserWindow.IsValid()) 
 	{
-		if (const TSharedRef<FCEFWebBrowserWindow> CefWebBrowserWindow = GetCefWebBrowserWindowChecked(WebBrowserWindow);
-			CefWebBrowserWindow->HasResizeBug() && CefWebBrowserWindow->IsCefResizeNeeded())
+		const TSharedPtr<FCEFWebBrowserWindow> CefWebBrowserWindow = StaticCastSharedPtr<FCEFWebBrowserWindow>(WebBrowserWindow);
+		if (CefWebBrowserWindow.IsValid() && CefWebBrowserWindow->HasResizeBug() && CefWebBrowserWindow->IsCefResizeNeeded())
 		{
 			CefWebBrowserWindow->RequestCefResize();
 		}
@@ -64,60 +42,6 @@ static FWebBrowserSingleton& GetWebBrowserSingletonChecked()
 #endif // WITH_EDITOR
 }
 
-/*virtual*/ int32 SWebBrowser::OnPaint(const FPaintArgs& Args, const FGeometry& AllottedGeometry, const FSlateRect& MyCullingRect, FSlateWindowElementList& OutDrawElements, int32 LayerId, const FWidgetStyle& InWidgetStyle, bool bParentEnabled) const
-{
-#if WITH_EDITOR
-#if WITH_CEF3
-	if (WebBrowserWindow.IsValid())
-	{
-		if (const TSharedPtr<FCEFWebBrowserWindow> CefWebBrowserWindow = GetCefWebBrowserWindowChecked(WebBrowserWindow);
-			CefWebBrowserWindow->HasResizeBug())
-		{
-			// NOTE - We won't end up here if the CEF version doesn't require it.
-			
-			if (const TSharedPtr<FCapturedCefBuffer>& CapturedCefBuffer = CefWebBrowserWindow->GetCapturedCefBuffer();
-				CapturedCefBuffer.IsValid())
-			{
-				if (CefWebBrowserWindow->HasCorrectNativeCefBuffer())
-				{
-					// The CEF web browser can display the correct texture itself. So clean-up the captured CEF buffer's paint objects, if it has any.
-					// But do this on the next tick to avoid a white flash.
-
-					if (CapturedCefBuffer->HasPaintObjects())
-					{
-						AsyncTask(ENamedThreads::GameThread, [CapturedCefBuffer]() -> void
-							{
-								if (CapturedCefBuffer.IsValid())
-								{
-									CapturedCefBuffer->ClearPaintObjects();
-								}
-							});
-					}
-				}
-				else if (CapturedCefBuffer->MaybeUpdatePaintObjects())
-				{
-					// The CEF web browser can't display the correct texture itself. So paint the captured CEF buffer instead.
-					// NOTE - We still need to call the usual OnPaint() to make the web browser update, but we ignore what it actually paints. Instead, we
-					// only properly paint the captured CEF buffer.
-			
-					FSlateWindowElementList IgnoredDrawElements(FSlateApplication::Get().FindWidgetWindow(AsShared()));
-					SCompoundWidget::OnPaint(Args, AllottedGeometry, MyCullingRect, IgnoredDrawElements, LayerId, InWidgetStyle, bParentEnabled);
-
-					const FWidgetStyle CompoundedWidgetStyle = FWidgetStyle(InWidgetStyle)
-						.BlendColorAndOpacityTint(GetColorAndOpacity())
-						.SetForegroundColor(ShouldBeEnabled(bParentEnabled) ? GetForegroundColor() : GetDisabledForegroundColor() );
-
-					return CapturedCefBuffer->PaintCentered(AllottedGeometry, OutDrawElements, LayerId, CompoundedWidgetStyle);
-				}
-			}
-		}
-	}
-#endif // WITH_CEF3
-#endif // WITH_EDITOR
-
-	return SCompoundWidget::OnPaint(Args, AllottedGeometry, MyCullingRect, OutDrawElements, LayerId, InWidgetStyle, bParentEnabled);
-}
-
 void SWebBrowser::Construct(const FArguments& InArgs, const TSharedPtr<IWebBrowserWindow>& InWebBrowserWindow)
 {
 	OnLoadCompleted = InArgs._OnLoadCompleted;
@@ -249,22 +173,15 @@ void SWebBrowser::Construct(const FArguments& InArgs, const TSharedPtr<IWebBrows
 
 
 #if WITH_CEF3
-
-
 	if (WebBrowserWindow.IsValid()) 
 	{
-		const bool bUseSmoothScrolling = 
-#if WITH_EDITOR
-			(GIsPlayInEditorWorld ? false : InArgs._UseSmoothResizing)
-#else
-			false
-#endif // WITH_EDITOR
-			;
-		
-		GetCefWebBrowserWindowChecked(WebBrowserWindow)->SetUseCapturedCefBuffer(bUseSmoothScrolling);
+		const TSharedPtr<FCEFWebBrowserWindow> CefWebBrowserWindow = StaticCastSharedPtr<FCEFWebBrowserWindow>(WebBrowserWindow);
+		if (CefWebBrowserWindow.IsValid())
+		{
+			CefWebBrowserWindow->SetCaptureCefFrames(InArgs._UseSmoothResizing);
+		}
 	}
 
-
 	// Make the IME (Input Method Editor) work with HTML text input boxes on web pages,
 	// for non-ASCII languages, like "Chinese (Traditional, Taiwan)".
 	if (FSlateApplication::IsInitialized())
@@ -274,8 +191,6 @@ void SWebBrowser::Construct(const FArguments& InArgs, const TSharedPtr<IWebBrows
 			BindInputMethodSystem(InputMethodSystem);
 		}
 	}
-	
-	
 #endif
 }
 
diff --git a/Engine/Source/Runtime/WebBrowser/Public/SWebBrowser.h b/Engine/Source/Runtime/WebBrowser/Public/SWebBrowser.h
--- a/Engine/Source/Runtime/WebBrowser/Public/SWebBrowser.h
+++ b/Engine/Source/Runtime/WebBrowser/Public/SWebBrowser.h
@@ -140,7 +140,6 @@ public:
 	virtual FReply OnFocusReceived(const FGeometry& MyGeometry, const FFocusEvent& InFocusEvent) override;
 
 	virtual void Tick(const FGeometry& AllottedGeometry, const double InCurrentTime, const float InDeltaTime) override;
-	virtual int32 OnPaint(const FPaintArgs& Args, const FGeometry& AllottedGeometry, const FSlateRect& MyCullingRect, FSlateWindowElementList& OutDrawElements, int32 LayerId, const FWidgetStyle& InWidgetStyle, bool bParentEnabled) const override;
 	
 	/**
 	 * Construct the widget.
