WebBrowser: Remove blanket Engine dependency on Mac also

The dependency was already made conditional on bCompileAgainstEngine at CL 48864048, but another "Engine" entry was left in a Mac-specific block elsewhere in the file. Remove that one too.

#rnx

#Synced-CL 50861308

diff --git a/Engine/Source/Runtime/WebBrowser/WebBrowser.Build.cs b/Engine/Source/Runtime/WebBrowser/WebBrowser.Build.cs
--- a/Engine/Source/Runtime/WebBrowser/WebBrowser.Build.cs
+++ b/Engine/Source/Runtime/WebBrowser/WebBrowser.Build.cs
@@ -38,7 +38,6 @@ public class WebBrowser : ModuleRules
 			PrivateDependencyModuleNames.AddRange(
 				new string[]
 				{
-					"Engine",
 					"JsonUtilities",
 					"WebBrowserTexture",
 					"Json"
