Arm64ec libcef.dll CEF crash fix
-Arm64ec is compatible with x64 and arm64ec, not with arm64.
-Not entirely sure what is happening inside libcef.dll, but somehow the arm64ec libcef.dll needs to run with a compatible EpicWebHelper.exe process.
#Synced-CL 51057591

diff --git a/Engine/Source/Runtime/WebBrowser/Private/WebBrowserSingleton.cpp b/Engine/Source/Runtime/WebBrowser/Private/WebBrowserSingleton.cpp
--- a/Engine/Source/Runtime/WebBrowser/Private/WebBrowserSingleton.cpp
+++ b/Engine/Source/Runtime/WebBrowser/Private/WebBrowserSingleton.cpp
@@ -77,10 +77,11 @@ THIRD_PARTY_INCLUDES_END
 #	if PLATFORM_WINDOWS && PLATFORM_CPU_ARM_FAMILY
 #		if PLATFORM_WINDOWS_ARM64EC
 #			define CEF3_RESOURCES_DIR CEF3_BIN_DIR TEXT("/Win64") CEF_VERSION_DIR CEF_VERSION_SUFFIX TEXT("/Resources")
+#			define CEF3_SUBPROCES_EXE TEXT("Binaries/Win64/EpicWebHelper.exe")
 #		else
 #			define CEF3_RESOURCES_DIR CEF3_BIN_DIR TEXT("/WinArm64") CEF_VERSION_DIR TEXT("/Resources")
+#			define CEF3_SUBPROCES_EXE TEXT("Binaries/Win64/EpicWebHelperarm64.exe")
 #		endif
-#		define CEF3_SUBPROCES_EXE TEXT("Binaries/Win64/EpicWebHelperarm64.exe")
 #	elif PLATFORM_WINDOWS && !PLATFORM_CPU_ARM_FAMILY
 #		define CEF3_RESOURCES_DIR CEF3_BIN_DIR TEXT("/Win64") CEF_VERSION_DIR CEF_VERSION_SUFFIX TEXT("/Resources")
 #		define CEF3_SUBPROCES_EXE TEXT("Binaries/Win64/EpicWebHelper.exe")
diff --git a/Engine/Source/Runtime/WebBrowser/WebBrowser.Build.cs b/Engine/Source/Runtime/WebBrowser/WebBrowser.Build.cs
--- a/Engine/Source/Runtime/WebBrowser/WebBrowser.Build.cs
+++ b/Engine/Source/Runtime/WebBrowser/WebBrowser.Build.cs
@@ -60,7 +60,12 @@ public class WebBrowser : ModuleRules
 
 			CircularlyReferencedDependentModules.Add("WebBrowserTexture");
 		}
-		
+
+		if (Target.Architecture == UnrealArch.Arm64ec)
+		{
+			PublicDefinitions.Add("PLATFORM_WINDOWS_ARM64EC=1");
+		}
+
 		if (Target.bCompileAgainstEngine)
 		{
 			if (Target.Platform == UnrealTargetPlatform.Win64)
@@ -101,7 +106,7 @@ public class WebBrowser : ModuleRules
 				}
 				else
 				{
-					if (Target.Architecture == UnrealArch.Arm64 || Target.Architecture == UnrealArch.Arm64ec)
+					if (Target.Architecture == UnrealArch.Arm64)
 					{
 						RuntimeDependencies.Add("$(EngineDir)/Binaries/" + Target.Platform.ToString() + "/EpicWebHelperarm64.exe");
 					}
