WebBrowser(CEF): Allow any HTML element with `contenteditable=true` to enable the IME when focused

and clean up the focus change handling to [de]activate the IME context instead of recreating it each time.

#rnx
#Synced-CL 51424673

diff --git a/Engine/Source/Runtime/WebBrowser/Private/CEF/CEFImeHandler.cpp b/Engine/Source/Runtime/WebBrowser/Private/CEF/CEFImeHandler.cpp
--- a/Engine/Source/Runtime/WebBrowser/Private/CEF/CEFImeHandler.cpp
+++ b/Engine/Source/Runtime/WebBrowser/Private/CEF/CEFImeHandler.cpp
@@ -10,6 +10,7 @@
 FCEFImeHandler::FCEFImeHandler(CefRefPtr<CefBrowser> Browser)
 	: InternalCefBrowser(Browser)
 	, TextInputMethodSystem(nullptr)
+	, bContextWasActiveWhenFocusLost(false)
 {
 
 }
@@ -51,27 +52,28 @@ void FCEFImeHandler::UnbindInputMethodSystem()
 	TextInputMethodSystem = nullptr;
 }
 
-void FCEFImeHandler::InitContext()
+void FCEFImeHandler::ActivateContext()
 {
 	if (!TextInputMethodSystem)
 	{
 		return;
 	}
 
-	// Clean up any existing context
-	DestroyContext();
-
-	TextInputMethodContext = FCEFTextInputMethodContext::Create(SharedThis(this));
-	if (TextInputMethodSystem)
+	if (!TextInputMethodContext.IsValid())
 	{
+		TextInputMethodContext = FCEFTextInputMethodContext::Create(SharedThis(this));
 		TextInputMethodChangeNotifier = TextInputMethodSystem->RegisterContext(TextInputMethodContext.ToSharedRef());
+		if (TextInputMethodChangeNotifier.IsValid())
+		{
+			TextInputMethodChangeNotifier->NotifyLayoutChanged(ITextInputMethodChangeNotifier::ELayoutChangeType::Created);
+		}
 	}
-	if (TextInputMethodChangeNotifier.IsValid())
+
+	TSharedRef<FCEFTextInputMethodContext> TextInputMethodContextRef = TextInputMethodContext.ToSharedRef();
+	if (!TextInputMethodSystem->IsActiveContext(TextInputMethodContextRef))
 	{
-		TextInputMethodChangeNotifier->NotifyLayoutChanged(ITextInputMethodChangeNotifier::ELayoutChangeType::Created);
+		TextInputMethodSystem->ActivateContext(TextInputMethodContextRef);
 	}
-
-	TextInputMethodSystem->ActivateContext(TextInputMethodContext.ToSharedRef());
 }
 
 void FCEFImeHandler::DeactivateContext()
@@ -119,10 +121,7 @@ bool FCEFImeHandler::HandleFocusChangedMessage(CefRefPtr<CefListValue> MessageAr
 	if (MessageArguments->GetSize() == 1 &&
 		MessageArguments->GetType(0) == VTYPE_STRING)
 	{
-		if (TextInputMethodContext.IsValid())
-		{
-			DestroyContext();
-		}
+		DeactivateContext();
 		return true;
 	}
 	else if (MessageArguments->GetSize() == 8 &&
@@ -136,15 +135,15 @@ bool FCEFImeHandler::HandleFocusChangedMessage(CefRefPtr<CefListValue> MessageAr
 		MessageArguments->GetType(7) == VTYPE_INT)
 	{
 		FString Type = WCHAR_TO_TCHAR(MessageArguments->GetString(0).ToWString().c_str());
-		FString Name = WCHAR_TO_TCHAR(MessageArguments->GetString(1).ToWString().c_str());
 		bool bIsEditable = MessageArguments->GetBool(2);
 
-		if (Type == TEXT("DOM_NODE_TYPE_ELEMENT") && 
-			(Name == TEXT("INPUT") || Name == TEXT("TEXTAREA")) &&
-			bIsEditable)
+		if (Type == TEXT("DOM_NODE_TYPE_ELEMENT") && bIsEditable)
+		{
+			ActivateContext();
+		}
+		else
 		{
-			// @todo: Make use of the bounds of the text box here to act as a fallback for the initial composition window pos.
-			InitContext();
+			DeactivateContext();
 		}
 		return true;
 	}
@@ -173,12 +172,16 @@ void FCEFImeHandler::SetFocus(bool bHasFocus)
 		return;
 	}
 
-	if (bHasFocus)
+	if (bContextWasActiveWhenFocusLost && bHasFocus)
 	{
-		TextInputMethodSystem->ActivateContext(TextInputMethodContext.ToSharedRef());
+		// Reactivate the context when we get focus back, for the case where the focused element in the page hasn't changed.
+		// If the user clicked on a different element to focus the browser, then we'll get a FocusChanged message and we'll adjust activation accordingly
+		ActivateContext();
 	}
 	else
 	{
+		// Save the current activated state of the context, so we can restore it if we get the focus back.
+		bContextWasActiveWhenFocusLost = TextInputMethodSystem->IsActiveContext(TextInputMethodContext.ToSharedRef());
 		DeactivateContext();
 	}
 }
diff --git a/Engine/Source/Runtime/WebBrowser/Private/CEF/CEFImeHandler.h b/Engine/Source/Runtime/WebBrowser/Private/CEF/CEFImeHandler.h
--- a/Engine/Source/Runtime/WebBrowser/Private/CEF/CEFImeHandler.h
+++ b/Engine/Source/Runtime/WebBrowser/Private/CEF/CEFImeHandler.h
@@ -93,7 +93,7 @@ private:
 		return InternalCefBrowser.get() != nullptr;
 	}
 
-	void InitContext();
+	void ActivateContext();
 	void DeactivateContext();
 	void DestroyContext();
 
@@ -113,6 +113,9 @@ private:
 	/** Notification interface object for IMEs */
 	TSharedPtr<ITextInputMethodChangeNotifier> TextInputMethodChangeNotifier;
 
+	/** Saves the active state of the TextInputMethodContext when the browser loses focus. */
+	bool bContextWasActiveWhenFocusLost;
+
 	// Allow IME context to access functions only it needs.
 	friend class FCEFTextInputMethodContext;
 };
