WebBrowser: Fix CEF <select> popup crash and rendering in accelerated paint path.

When accelerated paint was enabled for Windows (46518177), <select> dropdowns in SWebBrowser crashed, and if the crash was avoided they failed to render because UpdatableTextures[PET_POPUP] was created at 1x1 and never resized.

Fixed by storing the popup dimensions from OnPopupSize and pre-creating the sized texture. If the texture isn't ready on the first paint frame then it requests another via Invalidate(PET_POPUP).

Also removed the Y-flip on the popup viewport, it has been there since 2021 but accelerated paint has been disabled since until recently, and with the current CEF v128 it doesn't seem to need it.

#Synced-CL 52122091

diff --git a/Engine/Source/Runtime/WebBrowser/Private/CEF/CEFWebBrowserWindow.cpp b/Engine/Source/Runtime/WebBrowser/Private/CEF/CEFWebBrowserWindow.cpp
--- a/Engine/Source/Runtime/WebBrowser/Private/CEF/CEFWebBrowserWindow.cpp
+++ b/Engine/Source/Runtime/WebBrowser/Private/CEF/CEFWebBrowserWindow.cpp
@@ -2184,7 +2184,27 @@ void FCEFWebBrowserWindow::OnAcceleratedPaint(CefRenderHandler::PaintElementType
 	}
 #endif // UE_CEF_HAS_RESIZE_BUG
 
-	if (RHIRenderHelper && RHIRenderHelper->CopySharedTextureSync(UpdatableTextures[Type], SharedHandle, DirtyRect))
+	if (Type == PET_POPUP && RHIRenderHelper && PopupCefSize.X > 0 && PopupCefSize.Y > 0 && UpdatableTextures[Type] != nullptr)
+	{
+		const float DPIScale = FPlatformApplicationMisc::GetDPIScaleFactorAtPoint(PopupPosition.X, PopupPosition.Y);
+
+		const FIntPoint DevicePixelSize(
+			FMath::CeilToInt(PopupCefSize.X * DPIScale),
+			FMath::CeilToInt(PopupCefSize.Y * DPIScale)
+		);
+
+		const int32 CurW = UpdatableTextures[Type]->GetSlateResource()->GetWidth();
+		const int32 CurH = UpdatableTextures[Type]->GetSlateResource()->GetHeight();
+
+		if (CurW != DevicePixelSize.X || CurH != DevicePixelSize.Y)
+		{
+			RHIRenderHelper->ReleaseSlateUpdatableTexture(UpdatableTextures[Type]);
+			UpdatableTextures[Type] = RHIRenderHelper->CreateSlateUpdatableTexture(DevicePixelSize);
+		}
+	}
+
+	const bool bCopyOk = RHIRenderHelper && RHIRenderHelper->CopySharedTextureSync(UpdatableTextures[Type], SharedHandle, DirtyRect);
+	if (bCopyOk)
 	{
 		bNeedsRedraw = true;
 		if (Type == PET_POPUP && bShowPopupRequested)
@@ -2192,13 +2212,15 @@ void FCEFWebBrowserWindow::OnAcceleratedPaint(CefRenderHandler::PaintElementType
 			bShowPopupRequested = false;
 			bPopupHasFocus = true;
 
-			const float DPIScale = FPlatformApplicationMisc::GetDPIScaleFactorAtPoint(PopupPosition.X, PopupPosition.Y);
-			FIntPoint PopupSize = FIntPoint(UpdatableTextures[Type]->GetSlateResource()->GetWidth() / DPIScale, UpdatableTextures[Type]->GetSlateResource()->GetHeight() / DPIScale);
-
-			FIntRect PopupRect = FIntRect(PopupPosition, PopupPosition + PopupSize);
+			FIntRect PopupRect = FIntRect(PopupPosition, PopupPosition + PopupCefSize);
 			OnShowPopup().Broadcast(PopupRect);
 		}
 	}
+	else if (Type == PET_POPUP && bShowPopupRequested && InternalCefBrowser)
+	{
+		// Since the Texture RHI resource wasn't ready yet, we request another paint.
+		InternalCefBrowser->GetHost()->Invalidate(PET_POPUP);
+	}
 
 	bIsInitialized = true;
 	if (bNeedsRedraw)
@@ -3054,8 +3076,23 @@ void FCEFWebBrowserWindow::OnBrowserClosed()
 
 void FCEFWebBrowserWindow::SetPopupMenuPosition(CefRect CefPopupSize)
 {
-	// We only store the position, as the size will be provided ib the OnPaint call.
 	PopupPosition = FIntPoint(CefPopupSize.x, CefPopupSize.y);
+	PopupCefSize = FIntPoint(CefPopupSize.width, CefPopupSize.height);
+
+	// Pre-create the popup texture at the requested size
+	if (bUsingAcceleratedPaint && RHIRenderHelper && PopupCefSize.X > 0 && PopupCefSize.Y > 0)
+	{
+		const float DPIScale = FPlatformApplicationMisc::GetDPIScaleFactorAtPoint(PopupPosition.X, PopupPosition.Y);
+		const FIntPoint DevicePixelSize(
+			FMath::CeilToInt(PopupCefSize.X * DPIScale),
+			FMath::CeilToInt(PopupCefSize.Y * DPIScale));
+
+		if (UpdatableTextures[PET_POPUP] != nullptr)
+		{
+			RHIRenderHelper->ReleaseSlateUpdatableTexture(UpdatableTextures[PET_POPUP]);
+		}
+		UpdatableTextures[PET_POPUP] = RHIRenderHelper->CreateSlateUpdatableTexture(DevicePixelSize);
+	}
 }
 
 void FCEFWebBrowserWindow:: ShowPopupMenu(bool bShow)
@@ -3068,6 +3105,7 @@ void FCEFWebBrowserWindow:: ShowPopupMenu(bool bShow)
 	{
 		bPopupHasFocus = false;
 		bShowPopupRequested = false;
+		PopupCefSize = FIntPoint::ZeroValue;
 		OnDismissPopup().Broadcast();
 	}
 }
diff --git a/Engine/Source/Runtime/WebBrowser/Private/CEF/CEFWebBrowserWindow.h b/Engine/Source/Runtime/WebBrowser/Private/CEF/CEFWebBrowserWindow.h
--- a/Engine/Source/Runtime/WebBrowser/Private/CEF/CEFWebBrowserWindow.h
+++ b/Engine/Source/Runtime/WebBrowser/Private/CEF/CEFWebBrowserWindow.h
@@ -487,10 +487,10 @@ private:
 	void OnBrowserClosed();
 
 	/**
-	 * Called to set the popup menu location. Note that CEF also passes a size to this method,
-	 * which is ignored as the correct size is usually not known until inside OnPaint.
+	 * Called to set the popup menu position and size. For the accelerated paint path,
+	 * the size is used to pre-create the popup texture at the correct dimensions.
 	 *
- 	 * @param PopupSize The location of the popup widget.
+ 	 * @param PopupSize The location and size of the popup widget.
 	 */
 	void SetPopupMenuPosition(CefRect PopupSize);
 
@@ -794,6 +794,7 @@ private:
 	bool bSupportsMouseWheel;
 
 	FIntPoint PopupPosition;
+	FIntPoint PopupCefSize = FIntPoint::ZeroValue;
 	bool bShowPopupRequested;
 
 	/** This is set to true when reloading after render process crash. */
diff --git a/Engine/Source/Runtime/WebBrowser/Private/SWebBrowserView.cpp b/Engine/Source/Runtime/WebBrowser/Private/SWebBrowserView.cpp
--- a/Engine/Source/Runtime/WebBrowser/Private/SWebBrowserView.cpp
+++ b/Engine/Source/Runtime/WebBrowser/Private/SWebBrowserView.cpp
@@ -684,49 +684,34 @@ void SWebBrowserView::HandleShowPopup(const FIntRect& PopupSize)
 				.EnableGammaCorrection(false)
 				.EnableBlending(false)
 				.IgnoreTextureAlpha(true)
-#if WITH_CEF3
-				.RenderTransform(this, &SWebBrowserView::GetPopupRenderTransform)
-#endif
 				.Visibility(EVisibility::Visible);
 		MenuViewport = MakeShareable(new FWebBrowserViewport(BrowserWindow, true));
 	MenuContent->SetViewportInterface(MenuViewport.ToSharedRef());
+
 	FWidgetPath WidgetPath;
 	FSlateApplication::Get().GeneratePathToWidgetUnchecked(SharedThis(this), WidgetPath);
+
 	if (WidgetPath.IsValid())
 	{
 		TSharedRef< SWidget > MenuContentRef = MenuContent.ToSharedRef();
 		const FGeometry& BrowserGeometry = WidgetPath.Widgets.Last().Geometry;
 		const FVector2D NewPosition = BrowserGeometry.LocalToAbsolute(PopupSize.Min);
 
-
 		// Open the pop-up. The popup method will be queried from the widget path passed in.
 		TSharedPtr<IMenu> NewMenu = FSlateApplication::Get().PushMenu(SharedThis(this), WidgetPath, MenuContentRef, NewPosition, FPopupTransitionEffect( FPopupTransitionEffect::ComboButton ), false);
-		NewMenu->GetOnMenuDismissed().AddSP(this, &SWebBrowserView::HandleMenuDismissed);
-		PopupMenuPtr = NewMenu;
+
+		if (NewMenu.IsValid())
+		{
+			NewMenu->GetOnMenuDismissed().AddSP(this, &SWebBrowserView::HandleMenuDismissed);
+			PopupMenuPtr = NewMenu;
+		}
 	}
 
 }
 
 TOptional <FSlateRenderTransform> SWebBrowserView::GetPopupRenderTransform() const
 {
-	if (BrowserWindow.IsValid())
-	{
-#if !defined(DUMMY_WEB_BROWSER) && WITH_CEF3
-		TOptional<FSlateRenderTransform> LocalRenderTransform = FSlateRenderTransform();
-		if (static_cast<FWebBrowserWindow*>(BrowserWindow.Get())->UsingAcceleratedPaint())
-		{
-			// the accelerated renderer for CEF generates inverted textures (compared to the slate co-ord system), so flip it here
-			LocalRenderTransform = FSlateRenderTransform(Concatenate(FScale2D(1, -1), FVector2D(0, PopupMenuPtr.Pin()->GetContent()->GetDesiredSize().Y)));
-		}
-		return LocalRenderTransform;
-#else
-		return FSlateRenderTransform();
-#endif
-	}
-	else
-	{
-		return FSlateRenderTransform();
-	}
+	return FSlateRenderTransform();
 }
 
 void SWebBrowserView::HandleMenuDismissed(TSharedRef<IMenu>)
