WebBrowser: Add inbound drag-n-drop support

Forward the UE drag/drop events to the underlying webview so we can drop content into a web page.
Currently only implemented by the CEF backend for desktop platforms. The Apple WebKit backend already handles drag-n-drop natively.

#rnx
#Synced-CL 52253472

diff --git a/Engine/Source/Runtime/WebBrowser/Private/Android/AndroidWebBrowserWindow.cpp b/Engine/Source/Runtime/WebBrowser/Private/Android/AndroidWebBrowserWindow.cpp
--- a/Engine/Source/Runtime/WebBrowser/Private/Android/AndroidWebBrowserWindow.cpp
+++ b/Engine/Source/Runtime/WebBrowser/Private/Android/AndroidWebBrowserWindow.cpp
@@ -421,4 +421,26 @@ void FAndroidWebBrowserWindow::FloatingCloseButtonPressed()
 	}
 }
 
+void FAndroidWebBrowserWindow::OnDragEnter(const FGeometry& MyGeometry, const FDragDropEvent& DragDropEvent)
+{
+	// Not supported on this platform
+}
+
+FReply FAndroidWebBrowserWindow::OnDragOver(const FGeometry& MyGeometry, const FDragDropEvent& DragDropEvent)
+{
+	// Not supported on this platform
+	return FReply::Unhandled();
+}
+
+void FAndroidWebBrowserWindow::OnDragLeave(const FDragDropEvent& DragDropEvent)
+{
+	// Not supported on this platform
+}
+
+FReply FAndroidWebBrowserWindow::OnDrop(const FGeometry& MyGeometry, const FDragDropEvent& DragDropEvent)
+{
+	// Not supported on this platform
+	return FReply::Unhandled();
+}
+
 #endif // USE_ANDROID_JNI
diff --git a/Engine/Source/Runtime/WebBrowser/Private/Android/AndroidWebBrowserWindow.h b/Engine/Source/Runtime/WebBrowser/Private/Android/AndroidWebBrowserWindow.h
--- a/Engine/Source/Runtime/WebBrowser/Private/Android/AndroidWebBrowserWindow.h
+++ b/Engine/Source/Runtime/WebBrowser/Private/Android/AndroidWebBrowserWindow.h
@@ -79,6 +79,10 @@ public:
 	virtual FReply OnTouchGesture(const FGeometry& MyGeometry, const FPointerEvent& GestureEvent, bool bIsPopup) override;
 	virtual void OnFocus(bool SetFocus, bool bIsPopup) override;
 	virtual void OnCaptureLost() override;
+	virtual void OnDragEnter(const FGeometry& MyGeometry, const FDragDropEvent& DragDropEvent) override;
+	virtual FReply OnDragOver(const FGeometry& MyGeometry, const FDragDropEvent& DragDropEvent) override;
+	virtual void OnDragLeave(const FDragDropEvent& DragDropEvent) override;
+	virtual FReply OnDrop(const FGeometry& MyGeometry, const FDragDropEvent& DragDropEvent) override;
 	virtual bool CanGoBack() const override;
 	virtual void GoBack() override;
 	virtual bool CanGoForward() const override;
diff --git a/Engine/Source/Runtime/WebBrowser/Private/Apple/ApplePlatformWebBrowser.cpp b/Engine/Source/Runtime/WebBrowser/Private/Apple/ApplePlatformWebBrowser.cpp
--- a/Engine/Source/Runtime/WebBrowser/Private/Apple/ApplePlatformWebBrowser.cpp
+++ b/Engine/Source/Runtime/WebBrowser/Private/Apple/ApplePlatformWebBrowser.cpp
@@ -2150,4 +2150,26 @@ void FWebBrowserWindow::FloatingCloseButtonPressed()
 		OnFloatingCloseButtonPressed().Execute();
 	}
 }
+
+void FWebBrowserWindow::OnDragEnter(const FGeometry& MyGeometry, const FDragDropEvent& DragDropEvent)
+{
+	// Not supported on this platform
+}
+
+FReply FWebBrowserWindow::OnDragOver(const FGeometry& MyGeometry, const FDragDropEvent& DragDropEvent)
+{
+	// Not supported on this platform
+	return FReply::Unhandled();
+}
+
+void FWebBrowserWindow::OnDragLeave(const FDragDropEvent& DragDropEvent)
+{
+	// Not supported on this platform
+}
+
+FReply FWebBrowserWindow::OnDrop(const FGeometry& MyGeometry, const FDragDropEvent& DragDropEvent)
+{
+	// Not supported on this platform
+	return FReply::Unhandled();
+}
 #endif
diff --git a/Engine/Source/Runtime/WebBrowser/Private/Apple/ApplePlatformWebBrowser.h b/Engine/Source/Runtime/WebBrowser/Private/Apple/ApplePlatformWebBrowser.h
--- a/Engine/Source/Runtime/WebBrowser/Private/Apple/ApplePlatformWebBrowser.h
+++ b/Engine/Source/Runtime/WebBrowser/Private/Apple/ApplePlatformWebBrowser.h
@@ -186,6 +186,10 @@ public:
 	virtual FReply OnTouchGesture(const FGeometry& MyGeometry, const FPointerEvent& GestureEvent, bool bIsPopup) override;
 	virtual void OnFocus(bool SetFocus, bool bIsPopup) override;
 	virtual void OnCaptureLost() override;
+	virtual void OnDragEnter(const FGeometry& MyGeometry, const FDragDropEvent& DragDropEvent) override;
+	virtual FReply OnDragOver(const FGeometry& MyGeometry, const FDragDropEvent& DragDropEvent) override;
+	virtual void OnDragLeave(const FDragDropEvent& DragDropEvent) override;
+	virtual FReply OnDrop(const FGeometry& MyGeometry, const FDragDropEvent& DragDropEvent) override;
 	virtual bool CanGoBack() const override;
 	virtual void GoBack() override;
 	virtual bool CanGoForward() const override;
diff --git a/Engine/Source/Runtime/WebBrowser/Private/CEF/CEFWebBrowserWindow.cpp b/Engine/Source/Runtime/WebBrowser/Private/CEF/CEFWebBrowserWindow.cpp
--- a/Engine/Source/Runtime/WebBrowser/Private/CEF/CEFWebBrowserWindow.cpp
+++ b/Engine/Source/Runtime/WebBrowser/Private/CEF/CEFWebBrowserWindow.cpp
@@ -8,6 +8,7 @@
 #include "HAL/PlatformApplicationMisc.h"
 #include "Misc/CommandLine.h"
 #include "Misc/ConfigCacheIni.h"
+#include "PlatformHttp.h"
 #include "WebBrowserLog.h"
 
 #if WITH_CEF3
@@ -3156,4 +3157,145 @@ bool FCEFWebBrowserWindow::IsInDragRegion(const FIntPoint& Point)
 	return false;
 }
 
+void FCEFWebBrowserWindow::OnDragEnter(const FGeometry& MyGeometry, const FDragDropEvent& DragDropEvent)
+{
+	if (!IsValid())
+	{
+		return;
+	}
+
+	TSharedPtr<FExternalDragOperation> ExternalDragOp = DragDropEvent.GetOperationAs<FExternalDragOperation>();
+	if (!ExternalDragOp.IsValid())
+	{
+		return;
+	}
+
+	CefRefPtr<CefDragData> DragData = CefDragData::Create();
+
+	if (ExternalDragOp->HasFiles())
+	{
+		const TArray<FString>& Files = ExternalDragOp->GetFiles();
+		for (const FString& FilePath : Files)
+		{
+			DragData->AddFile(TCHAR_TO_WCHAR(*FilePath), CefString());
+		}
+	}
+
+	if (ExternalDragOp->HasText())
+	{
+		const FString& Text = ExternalDragOp->GetText();
+		if (Text.StartsWith(TEXT("file://")))
+		{
+			// Convert file URL to a local file path and add as a file:
+			// file://[/]C:/path/to/file -> C:/path/to/file (Windows)
+			// file://[/]path/to/file -> /path/to/file (Unix)
+			FString FilePath = Text.RightChop(7);
+			// Remove any left '/' in front of the Windows drive
+			if (FilePath.Len() >= 3 && FilePath[2] == TEXT(':'))
+			{
+				FilePath = FilePath.RightChop(1);
+			}
+			// Add any missing '/' in front of the Unix path
+			if (FilePath.Len() >= 2 && FilePath[1] != TEXT(':') && FilePath[0] != TEXT('/'))
+			{
+				FilePath.InsertAt(0, TEXT('/'));
+			}
+
+			// URL decode the path (handle %20 for spaces, etc.)
+			FilePath = FPlatformHttp::UrlDecode(FilePath);
+			FPaths::NormalizeFilename(FilePath);
+
+			DragData->AddFile(TCHAR_TO_WCHAR(*FilePath), CefString());
+		}
+		else if (Text.StartsWith(TEXT("http://")) || Text.StartsWith(TEXT("https://")))
+		{
+			DragData->SetLinkURL(TCHAR_TO_WCHAR(*Text));
+			DragData->SetLinkTitle(TCHAR_TO_WCHAR(*Text));
+		}
+		else
+		{
+			DragData->SetFragmentText(TCHAR_TO_WCHAR(*Text));
+		}
+	}
+
+	// Convert the mouse position to CEF coordinates
+	CefMouseEvent MouseEvent = GetCefMouseEvent(MyGeometry, DragDropEvent, false);
+
+	InternalCefBrowser->GetHost()->DragTargetDragEnter(DragData, MouseEvent, DRAG_OPERATION_COPY);
+}
+
+FReply FCEFWebBrowserWindow::OnDragOver(const FGeometry& MyGeometry, const FDragDropEvent& DragDropEvent)
+{
+	if (!IsValid())
+	{
+		return FReply::Unhandled();
+	}
+
+	TSharedPtr<FExternalDragOperation> ExternalDragOp = DragDropEvent.GetOperationAs<FExternalDragOperation>();
+	if (!ExternalDragOp.IsValid())
+	{
+		return FReply::Unhandled();
+	}
+
+	// Check if we have any supported drag data
+	if (!ExternalDragOp->HasFiles() && !ExternalDragOp->HasText())
+	{
+		return FReply::Unhandled();
+	}
+
+	// Convert the mouse position to CEF coordinates
+	CefMouseEvent MouseEvent = GetCefMouseEvent(MyGeometry, DragDropEvent, false);
+
+	// Allow all common operations - CEF will determine what's actually allowed based on the drag data
+	CefBrowserHost::DragOperationsMask AllowedOps = static_cast<CefBrowserHost::DragOperationsMask>(
+		DRAG_OPERATION_COPY | DRAG_OPERATION_MOVE | DRAG_OPERATION_LINK);
+
+	InternalCefBrowser->GetHost()->DragTargetDragOver(MouseEvent, AllowedOps);
+
+	return FReply::Handled();
+}
+
+void FCEFWebBrowserWindow::OnDragLeave(const FDragDropEvent& DragDropEvent)
+{
+	if (!IsValid())
+	{
+		return;
+	}
+
+	TSharedPtr<FExternalDragOperation> ExternalDragOp = DragDropEvent.GetOperationAs<FExternalDragOperation>();
+	if (!ExternalDragOp.IsValid())
+	{
+		return;
+	}
+
+	InternalCefBrowser->GetHost()->DragTargetDragLeave();
+}
+
+FReply FCEFWebBrowserWindow::OnDrop(const FGeometry& MyGeometry, const FDragDropEvent& DragDropEvent)
+{
+	if (!IsValid())
+	{
+		return FReply::Unhandled();
+	}
+
+	TSharedPtr<FExternalDragOperation> ExternalDragOp = DragDropEvent.GetOperationAs<FExternalDragOperation>();
+	if (!ExternalDragOp.IsValid())
+	{
+		return FReply::Unhandled();
+	}
+
+	// Check if we have any supported drag data
+	if (!ExternalDragOp->HasFiles() && !ExternalDragOp->HasText())
+	{
+		return FReply::Unhandled();
+	}
+
+	// Convert the mouse position to CEF coordinates
+	CefMouseEvent MouseEvent = GetCefMouseEvent(MyGeometry, DragDropEvent, false);
+
+	InternalCefBrowser->GetHost()->DragTargetDrop(MouseEvent);
+
+	return FReply::Handled();
+}
+
 #endif
diff --git a/Engine/Source/Runtime/WebBrowser/Private/CEF/CEFWebBrowserWindow.h b/Engine/Source/Runtime/WebBrowser/Private/CEF/CEFWebBrowserWindow.h
--- a/Engine/Source/Runtime/WebBrowser/Private/CEF/CEFWebBrowserWindow.h
+++ b/Engine/Source/Runtime/WebBrowser/Private/CEF/CEFWebBrowserWindow.h
@@ -142,6 +142,11 @@ public:
 	virtual FReply OnTouchGesture(const FGeometry& MyGeometry, const FPointerEvent& GestureEvent, bool bIsPopup) override;
 	virtual void OnFocus(bool SetFocus, bool bIsPopup) override;
 	virtual void OnCaptureLost() override;
+	virtual void OnDragEnter(const FGeometry& MyGeometry, const FDragDropEvent& DragDropEvent) override;
+	virtual FReply OnDragOver(const FGeometry& MyGeometry, const FDragDropEvent& DragDropEvent) override;
+	virtual void OnDragLeave(const FDragDropEvent& DragDropEvent) override;
+	virtual FReply OnDrop(const FGeometry& MyGeometry, const FDragDropEvent& DragDropEvent) override;
+
 	virtual bool CanGoBack() const override;
 	virtual void GoBack() override;
 	virtual bool CanGoForward() const override;
diff --git a/Engine/Source/Runtime/WebBrowser/Private/Native/NativeWebBrowserProxy.cpp b/Engine/Source/Runtime/WebBrowser/Private/Native/NativeWebBrowserProxy.cpp
--- a/Engine/Source/Runtime/WebBrowser/Private/Native/NativeWebBrowserProxy.cpp
+++ b/Engine/Source/Runtime/WebBrowser/Private/Native/NativeWebBrowserProxy.cpp
@@ -241,3 +241,25 @@ TSharedPtr<SWindow> FNativeWebBrowserProxy::GetParentWindow() const
 void FNativeWebBrowserProxy::SetParentWindow(TSharedPtr<SWindow> Window)
 {
 }
+
+void FNativeWebBrowserProxy::OnDragEnter(const FGeometry& MyGeometry, const FDragDropEvent& DragDropEvent)
+{
+	// Not supported on this platform
+}
+
+FReply FNativeWebBrowserProxy::OnDragOver(const FGeometry& MyGeometry, const FDragDropEvent& DragDropEvent)
+{
+	// Not supported on this platform
+	return FReply::Unhandled();
+}
+
+void FNativeWebBrowserProxy::OnDragLeave(const FDragDropEvent& DragDropEvent)
+{
+	// Not supported on this platform
+}
+
+FReply FNativeWebBrowserProxy::OnDrop(const FGeometry& MyGeometry, const FDragDropEvent& DragDropEvent)
+{
+	// Not supported on this platform
+	return FReply::Unhandled();
+}
diff --git a/Engine/Source/Runtime/WebBrowser/Private/Native/NativeWebBrowserProxy.h b/Engine/Source/Runtime/WebBrowser/Private/Native/NativeWebBrowserProxy.h
--- a/Engine/Source/Runtime/WebBrowser/Private/Native/NativeWebBrowserProxy.h
+++ b/Engine/Source/Runtime/WebBrowser/Private/Native/NativeWebBrowserProxy.h
@@ -55,6 +55,10 @@ public:
 	}
 	virtual void OnFocus(bool SetFocus, bool bIsPopup) override;
 	virtual void OnCaptureLost() override;
+	virtual void OnDragEnter(const FGeometry& MyGeometry, const FDragDropEvent& DragDropEvent) override;
+	virtual FReply OnDragOver(const FGeometry& MyGeometry, const FDragDropEvent& DragDropEvent) override;
+	virtual void OnDragLeave(const FDragDropEvent& DragDropEvent) override;
+	virtual FReply OnDrop(const FGeometry& MyGeometry, const FDragDropEvent& DragDropEvent) override;
 	virtual bool CanGoBack() const override;
 	virtual void GoBack() override;
 	virtual bool CanGoForward() const override;
diff --git a/Engine/Source/Runtime/WebBrowser/Private/SWebBrowserView.cpp b/Engine/Source/Runtime/WebBrowser/Private/SWebBrowserView.cpp
--- a/Engine/Source/Runtime/WebBrowser/Private/SWebBrowserView.cpp
+++ b/Engine/Source/Runtime/WebBrowser/Private/SWebBrowserView.cpp
@@ -296,6 +296,40 @@ FReply SWebBrowserView::OnFocusReceived(const FGeometry& MyGeometry, const FFocu
 	return Reply;
 }
 
+void SWebBrowserView::OnDragEnter(const FGeometry& MyGeometry, const FDragDropEvent& DragDropEvent)
+{
+	if (BrowserWindow.IsValid())
+	{
+		BrowserWindow->OnDragEnter(MyGeometry, DragDropEvent);
+	}
+}
+
+FReply SWebBrowserView::OnDragOver(const FGeometry& MyGeometry, const FDragDropEvent& DragDropEvent)
+{
+	if (BrowserWindow.IsValid())
+	{
+		return BrowserWindow->OnDragOver(MyGeometry, DragDropEvent);
+	}
+	return FReply::Unhandled();
+}
+
+void SWebBrowserView::OnDragLeave(const FDragDropEvent& DragDropEvent)
+{
+	if (BrowserWindow.IsValid())
+	{
+		BrowserWindow->OnDragLeave(DragDropEvent);
+	}
+}
+
+FReply SWebBrowserView::OnDrop(const FGeometry& MyGeometry, const FDragDropEvent& DragDropEvent)
+{
+	if (BrowserWindow.IsValid())
+	{
+		return BrowserWindow->OnDrop(MyGeometry, DragDropEvent);
+	}
+	return FReply::Unhandled();
+}
+
 void SWebBrowserView::HandleWindowActivated()
 {
 	if (BrowserViewport.IsValid())
diff --git a/Engine/Source/Runtime/WebBrowser/Public/IWebBrowserWindow.h b/Engine/Source/Runtime/WebBrowser/Public/IWebBrowserWindow.h
--- a/Engine/Source/Runtime/WebBrowser/Public/IWebBrowserWindow.h
+++ b/Engine/Source/Runtime/WebBrowser/Public/IWebBrowserWindow.h
@@ -4,6 +4,7 @@
 
 #include "CoreMinimal.h"
 #include "Input/CursorReply.h"
+#include "Input/DragAndDrop.h"
 #include "Input/Reply.h"
 #include "Widgets/SWindow.h"
 #include "SWebBrowser.h"
@@ -282,6 +283,39 @@ public:
 	/** Called when Capture lost */
 	virtual void OnCaptureLost() = 0;
 
+	/**
+	 * Notify the browser that a drag operation has entered the window.
+	 *
+	 * @param MyGeometry The Geometry of the browser
+	 * @param DragDropEvent Information about the drag event
+	 */
+	virtual void OnDragEnter(const FGeometry& MyGeometry, const FDragDropEvent& DragDropEvent) = 0;
+
+	/**
+	 * Notify the browser that a drag operation is moving over the window.
+	 *
+	 * @param MyGeometry The Geometry of the browser
+	 * @param DragDropEvent Information about the drag event
+	 * @return FReply indicating how the event was handled
+	 */
+	virtual FReply OnDragOver(const FGeometry& MyGeometry, const FDragDropEvent& DragDropEvent) = 0;
+
+	/**
+	 * Notify the browser that a drag operation has left the window.
+	 *
+	 * @param DragDropEvent Information about the drag event
+	 */
+	virtual void OnDragLeave(const FDragDropEvent& DragDropEvent) = 0;
+
+	/**
+	 * Notify the browser that a drag operation has been released on it.
+	 *
+	 * @param MyGeometry The Geometry of the browser
+	 * @param DragDropEvent Information about the drag event
+	 * @return FReply indicating how the event was handled
+	 */
+	virtual FReply OnDrop(const FGeometry& MyGeometry, const FDragDropEvent& DragDropEvent) = 0;
+
 	/**
 	 * Returns true if the browser can navigate backwards.
 	 */
diff --git a/Engine/Source/Runtime/WebBrowser/Public/SWebBrowserView.h b/Engine/Source/Runtime/WebBrowser/Public/SWebBrowserView.h
--- a/Engine/Source/Runtime/WebBrowser/Public/SWebBrowserView.h
+++ b/Engine/Source/Runtime/WebBrowser/Public/SWebBrowserView.h
@@ -3,6 +3,7 @@
 
 #include "CoreMinimal.h"
 #include "Widgets/DeclarativeSyntaxSupport.h"
+#include "Input/DragAndDrop.h"
 #include "Input/PopupMethodReply.h"
 #include "Widgets/SWidget.h"
 #include "Widgets/SCompoundWidget.h"
@@ -173,6 +174,11 @@ public:
 
 	virtual FReply OnFocusReceived(const FGeometry& MyGeometry, const FFocusEvent& InFocusEvent) override;
 
+	WEBBROWSER_API virtual void OnDragEnter(const FGeometry& MyGeometry, const FDragDropEvent& DragDropEvent) override;
+	WEBBROWSER_API virtual FReply OnDragOver(const FGeometry& MyGeometry, const FDragDropEvent& DragDropEvent) override;
+	WEBBROWSER_API virtual void OnDragLeave(const FDragDropEvent& DragDropEvent) override;
+	WEBBROWSER_API virtual FReply OnDrop(const FGeometry& MyGeometry, const FDragDropEvent& DragDropEvent) override;
+
 	/**
 	 * Construct the widget.
 	 *
