WebBrowser: Fix popups location in scaled game viewports

The popup rect and mouse forwarding must also take viewport scaling into consideration, for cases where a WebBrowser is added to a game viewport with DPI scaling rules. This is separate from the OS DPI scale factor.
Also simplify popup texture resizing, we can do it synchronously inside SetPopupMenuPosition so no need to duplicate the code in the Paint methods. This also fixes the non-accelerated mode.
Remove method that's not used anymore after CL52210239.

#rnx
#Synced-CL 52380124

diff --git a/Engine/Source/Runtime/WebBrowser/Private/CEF/CEFWebBrowserWindow.cpp b/Engine/Source/Runtime/WebBrowser/Private/CEF/CEFWebBrowserWindow.cpp
--- a/Engine/Source/Runtime/WebBrowser/Private/CEF/CEFWebBrowserWindow.cpp
+++ b/Engine/Source/Runtime/WebBrowser/Private/CEF/CEFWebBrowserWindow.cpp
@@ -393,6 +393,7 @@ FCEFWebBrowserWindow::FCEFWebBrowserWindow(CefRefPtr<CefBrowser> InBrowser, CefR
 	, CurrentUrl(InUrl)
 	, ViewportSize(FIntPoint::ZeroValue)
 	, ViewportDPIScaleFactor(1.0f)
+	, ViewportScaleFactor(1.0f)
 	, bIsClosing(false)
 	, bIsInitialized(false)
 	, ContentsToLoad(InContentsToLoad)
@@ -414,6 +415,7 @@ FCEFWebBrowserWindow::FCEFWebBrowserWindow(CefRefPtr<CefBrowser> InBrowser, CefR
 	, bMainHasFocus(false)
 	, bPopupHasFocus(false)
 	, bSupportsMouseWheel(true)
+	, bShowPopupRequested(false)
 	, bRecoverFromRenderProcessCrash(false)
 	, ErrorCode(0)
 	, bDeferNavigations(false)
@@ -466,7 +468,7 @@ void FCEFWebBrowserWindow::ReleaseTextures()
 			{
 				if (FSlateApplication::IsInitialized())
 				{
-					if (RHIRenderHelper)
+					if (bUsingAcceleratedPaint && RHIRenderHelper)
 					{
 						RHIRenderHelper->ReleaseSlateUpdatableTexture(TextureToRelease);
 					}
@@ -482,7 +484,7 @@ void FCEFWebBrowserWindow::ReleaseTextures()
 				{
 					if (FSlateApplication::IsInitialized())
 					{
-						if (RHIRenderHelper)
+						if (bUsingAcceleratedPaint && RHIRenderHelper)
 						{
 							RHIRenderHelper->ReleaseSlateUpdatableTexture(TextureToRelease);
 						}
@@ -683,7 +685,7 @@ void FCEFWebBrowserWindow::SetViewportSize(FIntPoint WindowSize, FIntPoint Windo
 
 			// The lines below create a new window texture with the new size, also copying the old texture contents 1:1 into the new one, so they don't look stretched.
 			// But, in the non-accelerated case we only want to use CapturedCefFrame's contents if they're from a recent frame (ie. an ongoing drag-resize)
-			if (RHIRenderHelper)
+			if (bUsingAcceleratedPaint && RHIRenderHelper)
 			{
 				UpdatableTextures[PET_VIEW] = RHIRenderHelper->ResizeSlateUpdatableTexture(UpdatableTextures[PET_VIEW], WindowSizeScaled);
 			}
@@ -2093,11 +2095,6 @@ void FCEFWebBrowserWindow::OnPaint(CefRenderHandler::PaintElementType Type, cons
 			{
 				bShowPopupRequested = false;
 				bPopupHasFocus = true;
-
-				const float DPIScale = FPlatformApplicationMisc::GetDPIScaleFactorAtPoint(PopupPosition.X, PopupPosition.Y);
-				FIntPoint PopupSize = FIntPoint(Width / DPIScale, Height / DPIScale);
-
-				FIntRect PopupRect = FIntRect(PopupPosition, PopupPosition + PopupSize);
 				OnShowPopup().Broadcast(PopupRect);
 			}
 		}
@@ -2185,43 +2182,16 @@ void FCEFWebBrowserWindow::OnAcceleratedPaint(CefRenderHandler::PaintElementType
 	}
 #endif // UE_CEF_HAS_RESIZE_BUG
 
-	if (Type == PET_POPUP && RHIRenderHelper && PopupCefSize.X > 0 && PopupCefSize.Y > 0 && UpdatableTextures[Type] != nullptr)
-	{
-		const float DPIScale = FPlatformApplicationMisc::GetDPIScaleFactorAtPoint(PopupPosition.X, PopupPosition.Y);
-
-		const FIntPoint DevicePixelSize(
-			FMath::CeilToInt(PopupCefSize.X * DPIScale),
-			FMath::CeilToInt(PopupCefSize.Y * DPIScale)
-		);
-
-		const int32 CurW = UpdatableTextures[Type]->GetSlateResource()->GetWidth();
-		const int32 CurH = UpdatableTextures[Type]->GetSlateResource()->GetHeight();
-
-		if (CurW != DevicePixelSize.X || CurH != DevicePixelSize.Y)
-		{
-			RHIRenderHelper->ReleaseSlateUpdatableTexture(UpdatableTextures[Type]);
-			UpdatableTextures[Type] = RHIRenderHelper->CreateSlateUpdatableTexture(DevicePixelSize);
-		}
-	}
-
-	const bool bCopyOk = RHIRenderHelper && RHIRenderHelper->CopySharedTextureSync(UpdatableTextures[Type], SharedHandle, DirtyRect);
-	if (bCopyOk)
+	if (RHIRenderHelper && RHIRenderHelper->CopySharedTextureSync(UpdatableTextures[Type], SharedHandle, DirtyRect))
 	{
 		bNeedsRedraw = true;
 		if (Type == PET_POPUP && bShowPopupRequested)
 		{
 			bShowPopupRequested = false;
 			bPopupHasFocus = true;
-
-			FIntRect PopupRect = FIntRect(PopupPosition, PopupPosition + PopupCefSize);
 			OnShowPopup().Broadcast(PopupRect);
 		}
 	}
-	else if (Type == PET_POPUP && bShowPopupRequested && InternalCefBrowser)
-	{
-		// Since the Texture RHI resource wasn't ready yet, we request another paint.
-		InternalCefBrowser->GetHost()->Invalidate(PET_POPUP);
-	}
 
 	bIsInitialized = true;
 	if (bNeedsRedraw)
@@ -2655,20 +2625,19 @@ int32 FCEFWebBrowserWindow::GetCefMouseModifiers(const FPointerEvent& InMouseEve
 
 CefMouseEvent FCEFWebBrowserWindow::GetCefMouseEvent(const FGeometry& MyGeometry, const FPointerEvent& MouseEvent, bool bIsPopup)
 {
-	CefMouseEvent Event;
+	FVector2D LocalPos = MyGeometry.AbsoluteToLocal(MouseEvent.GetScreenSpacePosition());
 
-	float DPIScale = MyGeometry.Scale;
-	if (TSharedPtr<SWindow> ParentWindowPtr = ParentWindow.Pin())
+	// convert from game viewport coords to webview coords
+	if (bIsPopup)
 	{
-		DPIScale /= ParentWindowPtr->GetNativeWindow()->GetDPIScaleFactor();
+		LocalPos += FVector2D(PopupRect.Min) * ViewportScaleFactor;
 	}
-
-
-	FVector2D LocalPos = MyGeometry.AbsoluteToLocal(MouseEvent.GetScreenSpacePosition()) * DPIScale;
-	if (bIsPopup)
+	else
 	{
-		LocalPos += PopupPosition;
+		LocalPos *= ViewportScaleFactor;
 	}
+
+	CefMouseEvent Event;
 	Event.x = LocalPos.X;
 	Event.y = LocalPos.Y;
 	Event.modifiers = GetCefMouseModifiers(MouseEvent);
@@ -2736,6 +2705,9 @@ bool FCEFWebBrowserWindow::CanSupportAcceleratedPaint()
 
 void FCEFWebBrowserWindow::UpdateCachedGeometry(const FGeometry& AllottedGeometry)
 {
+	// Extract the game viewport scale factor from the global Scale, which also includes the DPI scale factor from the OS
+	ViewportScaleFactor = AllottedGeometry.Scale / ViewportDPIScaleFactor;
+
 #if !PLATFORM_LINUX
 	// Forward along the geometry for use by IME
 	Ime->UpdateCachedGeometry(AllottedGeometry);
@@ -3077,22 +3049,15 @@ void FCEFWebBrowserWindow::OnBrowserClosed()
 
 void FCEFWebBrowserWindow::SetPopupMenuPosition(CefRect CefPopupSize)
 {
-	PopupPosition = FIntPoint(CefPopupSize.x, CefPopupSize.y);
-	PopupCefSize = FIntPoint(CefPopupSize.width, CefPopupSize.height);
+	// The popup position is in webview coords, convert to game viewport coords so it shows up in the right place
+	const FIntPoint PopupPosition = (FVector2D(CefPopupSize.x, CefPopupSize.y) / ViewportScaleFactor).IntPoint();
+	PopupRect = FIntRect(PopupPosition, PopupPosition + FIntPoint(CefPopupSize.width, CefPopupSize.height));
 
-	// Pre-create the popup texture at the requested size
-	if (bUsingAcceleratedPaint && RHIRenderHelper && PopupCefSize.X > 0 && PopupCefSize.Y > 0)
+	// Resize the popup texture if using accelerated rendering, taking OS DPI scaling into account
+	if (bUsingAcceleratedPaint && RHIRenderHelper)
 	{
-		const float DPIScale = FPlatformApplicationMisc::GetDPIScaleFactorAtPoint(PopupPosition.X, PopupPosition.Y);
-		const FIntPoint DevicePixelSize(
-			FMath::CeilToInt(PopupCefSize.X * DPIScale),
-			FMath::CeilToInt(PopupCefSize.Y * DPIScale));
-
-		if (UpdatableTextures[PET_POPUP] != nullptr)
-		{
-			RHIRenderHelper->ReleaseSlateUpdatableTexture(UpdatableTextures[PET_POPUP]);
-		}
-		UpdatableTextures[PET_POPUP] = RHIRenderHelper->CreateSlateUpdatableTexture(DevicePixelSize);
+		const FIntPoint TextureSize = (FVector2D(CefPopupSize.width, CefPopupSize.height) * ViewportDPIScaleFactor).IntPoint();
+		UpdatableTextures[PET_POPUP] = RHIRenderHelper->ResizeSlateUpdatableTexture(UpdatableTextures[PET_POPUP], TextureSize);
 	}
 }
 
@@ -3106,7 +3071,7 @@ void FCEFWebBrowserWindow:: ShowPopupMenu(bool bShow)
 	{
 		bPopupHasFocus = false;
 		bShowPopupRequested = false;
-		PopupCefSize = FIntPoint::ZeroValue;
+		PopupRect = FIntRect();
 		OnDismissPopup().Broadcast();
 	}
 }
diff --git a/Engine/Source/Runtime/WebBrowser/Private/CEF/CEFWebBrowserWindow.h b/Engine/Source/Runtime/WebBrowser/Private/CEF/CEFWebBrowserWindow.h
--- a/Engine/Source/Runtime/WebBrowser/Private/CEF/CEFWebBrowserWindow.h
+++ b/Engine/Source/Runtime/WebBrowser/Private/CEF/CEFWebBrowserWindow.h
@@ -682,9 +682,12 @@ private:
 	/** Current position of this window. */
 	FIntPoint ViewportPos;
 
-	/** Current DPI scale factor of this window. */
+	/** Current DPI scale factor of this window (from the OS). */
 	float ViewportDPIScaleFactor;
 
+	/** Current scale factor of this window (from the game viewport). */
+	float ViewportScaleFactor;
+
 	/** Whether this window is closing. */
 	bool bIsClosing;
 
@@ -798,8 +801,7 @@ private:
 
 	bool bSupportsMouseWheel;
 
-	FIntPoint PopupPosition;
-	FIntPoint PopupCefSize = FIntPoint::ZeroValue;
+	FIntRect PopupRect;
 	bool bShowPopupRequested;
 
 	/** This is set to true when reloading after render process crash. */
diff --git a/Engine/Source/Runtime/WebBrowser/Private/SWebBrowserView.cpp b/Engine/Source/Runtime/WebBrowser/Private/SWebBrowserView.cpp
--- a/Engine/Source/Runtime/WebBrowser/Private/SWebBrowserView.cpp
+++ b/Engine/Source/Runtime/WebBrowser/Private/SWebBrowserView.cpp
@@ -743,11 +743,6 @@ void SWebBrowserView::HandleShowPopup(const FIntRect& PopupSize)
 
 }
 
-TOptional <FSlateRenderTransform> SWebBrowserView::GetPopupRenderTransform() const
-{
-	return FSlateRenderTransform();
-}
-
 void SWebBrowserView::HandleMenuDismissed(TSharedRef<IMenu>)
 {
 	PopupMenuPtr.Reset();
diff --git a/Engine/Source/Runtime/WebBrowser/Public/SWebBrowserView.h b/Engine/Source/Runtime/WebBrowser/Public/SWebBrowserView.h
--- a/Engine/Source/Runtime/WebBrowser/Public/SWebBrowserView.h
+++ b/Engine/Source/Runtime/WebBrowser/Public/SWebBrowserView.h
@@ -384,7 +384,6 @@ private:
 	WEBBROWSER_API bool HandleDrag(const FPointerEvent& MouseEvent);
 	WEBBROWSER_API void HandleConsoleMessage(const FString& Message, const FString& Source, int32 Line, EWebBrowserConsoleLogSeverity Serverity);
 
-	WEBBROWSER_API TOptional<FSlateRenderTransform> GetPopupRenderTransform() const;
 private:
 
 	/** Interface for dealing with a web browser window. */
