small fixes related to ptr capture going on the wrong thread and making browser creation without screenshot support possible by removing a check

#Synced-CL 52243953

diff --git a/Engine/Source/Runtime/WebBrowser/Private/Apple/ApplePlatformWebBrowser.cpp b/Engine/Source/Runtime/WebBrowser/Private/Apple/ApplePlatformWebBrowser.cpp
--- a/Engine/Source/Runtime/WebBrowser/Private/Apple/ApplePlatformWebBrowser.cpp
+++ b/Engine/Source/Runtime/WebBrowser/Private/Apple/ApplePlatformWebBrowser.cpp
@@ -81,6 +81,9 @@ class SAppleWebBrowserWidget : public SViewport
 		WebViewWrapper = [AppleWebViewWrapper alloc];
 		[WebViewWrapper create : this userAgentApplication : Args._UserAgentApplication.GetNSString() useTransparency : Args._UseTransparency  enableFloatingCloseButton : bEnableFloatingCloseButton];
 
+		TSharedRef<SAppleWebBrowserWidget> SharedThis = 
+			StaticCastSharedPtr<SAppleWebBrowserWidget>(AsShared().ToSharedPtr()).ToSharedRef();
+		
 #if !PLATFORM_TVOS
 		OnSafeFrameChangedEventHandle = FCoreDelegates::OnSafeFrameChangedEvent.AddLambda( [this]()
 										{
@@ -131,13 +134,16 @@ class SAppleWebBrowserWidget : public SViewport
 			}
 		}
 		
-		check(![WebViewWrapper isAllowingScreenshots] || WebBrowserBrush.IsValid());
+		if([WebViewWrapper isAllowingScreenshots] && !WebBrowserBrush.IsValid())
+		{
+			UE_LOG(LogTemp, Warning, TEXT("SAppleWebBrowserWidget is allowing screenshots but the requisite materials were not found!"));
+		}
+		else if([WebViewWrapper isAllowingScreenshots])
+		{
+			FSlateApplication::Get().OnPreScreenshot().AddSP(SharedThis, &SAppleWebBrowserWidget::OnTakingScreenshot);
+		}
 #endif
 		
-		TSharedRef<SAppleWebBrowserWidget> SharedThis = 
-			StaticCastSharedPtr<SAppleWebBrowserWidget>(AsShared().ToSharedPtr()).ToSharedRef();
-		
-		FSlateApplication::Get().OnPreScreenshot().AddSP(SharedThis, &SAppleWebBrowserWidget::OnTakingScreenshot);
 		FCoreDelegates::OnEnginePreExit.AddSP(SharedThis, &SAppleWebBrowserWidget::OnEnginePreExit);
 		
 		LoadURL(Args._InitialURL);
@@ -791,7 +797,17 @@ class SAppleWebBrowserWidget : public SViewport
 	{
 		if(SetFocus && !Is3D()) 
 		{
-			FSlateApplication::Get().ReleaseAllPointerCapture();
+			if(IsInGameThread())
+			{
+				FSlateApplication::Get().ReleaseAllPointerCapture();
+			}
+			else
+			{
+				FGraphEventRef PointerReleaseEvent = FFunctionGraphTask::CreateAndDispatchWhenReady([]() {
+					FSlateApplication::Get().ReleaseAllPointerCapture();
+				}, TStatId(), nullptr, ENamedThreads::GameThread);
+				FTaskGraphInterface::Get().WaitUntilTaskCompletes(PointerReleaseEvent);
+			}
 		}
 		else if(!FSlateApplication::Get().IsActive() && bFocused) 
 		{
