WebBrowser: Fix crashes when third-party CEF plugins prevent successful module initialization

The singleton was still trying to tick CEF even when initialization failed. Check bCEFInitialized instead of bAllowCEF, as the latter doesn't account for failed initialization.
Also prevent CEF startup if we detect a version mismatch, the API isn't guaranteed to be compatible in that case.

#rnx
#14145

#Synced-CL 52469830

diff --git a/Engine/Source/Runtime/WebBrowser/Private/WebBrowserModule.cpp b/Engine/Source/Runtime/WebBrowser/Private/WebBrowserModule.cpp
--- a/Engine/Source/Runtime/WebBrowser/Private/WebBrowserModule.cpp
+++ b/Engine/Source/Runtime/WebBrowser/Private/WebBrowserModule.cpp
@@ -38,6 +38,7 @@ public:
 private:
 #if WITH_CEF3
 	bool bLoadedCEFModule = false;
+	bool bMatchingCEFVersion = false;
 #if PLATFORM_MAC
 	// Dynamically load the CEF framework library.
 	CefScopedLibraryLoader *CEFLibraryLoader = nullptr;
@@ -78,6 +79,10 @@ void FWebBrowserModule::StartupModule()
 	{
 		UE_LOG(LogWebBrowser, Warning, TEXT("CEF3 loaded version mismatch! Module was built against %i.%i.%i.%i, check if library loading path is correct"), CEF_VERSION_MAJOR, CEF_VERSION_MINOR, CEF_VERSION_PATCH, CEF_COMMIT_NUMBER);
 	}
+	else
+	{
+		bMatchingCEFVersion = true;
+	}
 #endif
 }
 
@@ -91,6 +96,8 @@ void FWebBrowserModule::ShutdownModule()
 
 #if WITH_CEF3
 	CEF3Utils::UnloadCEF3Modules();
+	bLoadedCEFModule = false;
+	bMatchingCEFVersion = false;
 #if PLATFORM_MAC
 	delete CEFLibraryLoader;
 	CEFLibraryLoader = nullptr;
@@ -121,7 +128,7 @@ IWebBrowserSingleton* FWebBrowserModule::GetSingleton()
 bool FWebBrowserModule::IsWebModuleAvailable() const
 {
 #if WITH_CEF3
-	return bLoadedCEFModule;
+	return bLoadedCEFModule && bMatchingCEFVersion;
 #else
 	return true;
 #endif
diff --git a/Engine/Source/Runtime/WebBrowser/Private/WebBrowserSingleton.cpp b/Engine/Source/Runtime/WebBrowser/Private/WebBrowserSingleton.cpp
--- a/Engine/Source/Runtime/WebBrowser/Private/WebBrowserSingleton.cpp
+++ b/Engine/Source/Runtime/WebBrowser/Private/WebBrowserSingleton.cpp
@@ -282,6 +282,8 @@ static bool IsValidCEFLockFile(const FString& LockFile)
 FWebBrowserSingleton::FWebBrowserSingleton(const FWebBrowserInitSettings& WebBrowserInitSettings)
 #if WITH_CEF3
 	: WebBrowserWindowFactory(MakeShareable(new FWebBrowserWindowFactory()))
+	, bCEFInitialized(false)
+	, bTaskQueueFlushed(false)
 #else
 	: WebBrowserWindowFactory(MakeShareable(new FNoWebBrowserWindowFactory()))
 #if PLATFORM_IOS || PLATFORM_MAC || (PLATFORM_ANDROID && USE_ANDROID_JNI)
@@ -291,9 +293,6 @@ FWebBrowserSingleton::FWebBrowserSingleton(const FWebBrowserInitSettings& WebBro
 	, bDevToolsShortcutEnabled(UE_BUILD_DEBUG)
 	, bJSBindingsToLoweringEnabled(true)
 	, bAppIsFocused(false)
-#if WITH_CEF3
-	, bCEFInitialized(false)
-#endif
 #if WITH_ENGINE
 	, DefaultMaterial(FSoftObjectPath(FString(TEXT("/Engine/WebBrowser/WebTexture_M.WebTexture_M"))))
 	, DefaultTranslucentMaterial(FSoftObjectPath(FString(TEXT("/Engine/WebBrowser/WebTexture_TM.WebTexture_TM"))))
@@ -303,8 +302,9 @@ FWebBrowserSingleton::FWebBrowserSingleton(const FWebBrowserInitSettings& WebBro
 	
 	// Only enable CEF if we have CEF3, we are not running a commandlet without rendering (e.g. cooking assets) and it has not been explicitly disabled
 	// Disallow CEF if we never plan on rendering, ie, with CanEverRender. This includes servers
-	bAllowCEF = (!IsRunningCommandlet() || (IsAllowCommandletRendering() && FParse::Param(FCommandLine::Get(), TEXT("AllowCommandletCEF")))) &&
-				FApp::CanEverRender() && !FParse::Param(FCommandLine::Get(), TEXT("nocef"));
+	bool bAllowCEF = (!IsRunningCommandlet() || (IsAllowCommandletRendering() && FParse::Param(FCommandLine::Get(), TEXT("AllowCommandletCEF")))) &&
+				FApp::CanEverRender() && !FParse::Param(FCommandLine::Get(), TEXT("nocef")) &&
+				IWebBrowserModule::Get().IsWebModuleAvailable();
 	if (bAllowCEF)
 	{
 		// The FWebBrowserSingleton must be initialized on the game thread
@@ -511,14 +511,14 @@ FWebBrowserSingleton::FWebBrowserSingleton(const FWebBrowserInitSettings& WebBro
 void FWebBrowserSingleton::WaitForTaskQueueFlush()
 {
 	// Keep pumping messages until we see the one below clear the queue
-	bTaskFinished = false;
-	CefPostTask(TID_UI, new FCEFBrowserClosureTask(nullptr, [=, this]()
+	bTaskQueueFlushed = false;
+	CefPostTask(TID_UI, new FCEFBrowserClosureTask(nullptr, [this]()
 		{
-			bTaskFinished = true;
+			bTaskQueueFlushed = true;
 		}));
 
 	const double StartWaitAppTime = FPlatformTime::Seconds();
-	while (!bTaskFinished)
+	while (!bTaskQueueFlushed)
 	{
 		FPlatformProcess::Sleep(0.01);
 		// CEF needs the windows message pump run to be able to finish closing a browser, so run it manually here
@@ -540,10 +540,7 @@ void FWebBrowserSingleton::WaitForTaskQueueFlush()
 FWebBrowserSingleton::~FWebBrowserSingleton()
 {
 #if WITH_CEF3
-	if (!bCEFInitialized)
-		return; // CEF failed to init so don't crash trying to shut it down
-
-	if (bAllowCEF)
+	if (bCEFInitialized)
 	{
 		{
 			FScopeLock Lock(&WindowInterfacesCS);
@@ -598,7 +595,7 @@ FWebBrowserSingleton::~FWebBrowserSingleton()
 bool FWebBrowserSingleton::IsShuttingDown() const
 {
 #if WITH_CEF3
-	return bAllowCEF && !CEFBrowserApp;
+	return bCEFInitialized && !CEFBrowserApp;
 #else
 	return false;
 #endif
@@ -615,7 +612,7 @@ TSharedPtr<IWebBrowserWindow> FWebBrowserSingleton::CreateBrowserWindow(
 	)
 {
 #if WITH_CEF3
-	if (bAllowCEF)
+	if (bCEFInitialized)
 	{
 		TOptional<FString> ContentsToLoad;
 
@@ -647,7 +644,7 @@ TSharedPtr<IWebBrowserWindow> FWebBrowserSingleton::CreateBrowserWindow(const FC
 	}
 
 #if WITH_CEF3
-	if (bAllowCEF)
+	if (bCEFInitialized)
 	{
 		// Information used when creating the native window.
 		CefWindowInfo WindowInfo;
@@ -844,7 +841,7 @@ bool FWebBrowserSingleton::Tick(float DeltaTime)
 	QUICK_SCOPE_CYCLE_COUNTER(STAT_FWebBrowserSingleton_Tick);
 
 #if WITH_CEF3
-	if (bAllowCEF)
+	if (bCEFInitialized)
 	{
 		{
 			FScopeLock Lock(&WindowInterfacesCS);
@@ -998,7 +995,7 @@ TSharedPtr<IWebBrowserCookieManager> FWebBrowserSingleton::GetCookieManager(TOpt
 	if (ContextId.IsSet())
 	{
 #if WITH_CEF3
-		if (bAllowCEF)
+		if (bCEFInitialized)
 		{
 			const CefRefPtr<CefRequestContext>* ExistingContext = RequestContexts.Find(ContextId.GetValue());
 
@@ -1130,7 +1127,7 @@ void FWebBrowserSingleton::ClearOldCacheFolders(const FString &CachePathRoot, co
 bool FWebBrowserSingleton::RegisterContext(const FBrowserContextSettings& Settings)
 {
 #if WITH_CEF3
-	if (bAllowCEF)
+	if (bCEFInitialized)
 	{
 		const CefRefPtr<CefRequestContext>* ExistingContext = RequestContexts.Find(Settings.Id);
 
@@ -1167,7 +1164,7 @@ bool FWebBrowserSingleton::UnregisterContext(const FString& ContextId)
 {
 #if WITH_CEF3
 	bool bFoundContext = false;
-	if (bAllowCEF)
+	if (bCEFInitialized)
 	{
 		UE_LOG(LogWebBrowser, Log, TEXT("Unregistering ContextId=%s."), *ContextId);
 
@@ -1195,7 +1192,7 @@ bool FWebBrowserSingleton::UnregisterContext(const FString& ContextId)
 bool FWebBrowserSingleton::RegisterSchemeHandlerFactory(FString Scheme, FString Domain, IWebBrowserSchemeHandlerFactory* WebBrowserSchemeHandlerFactory)
 {
 #if WITH_CEF3
-	if (bAllowCEF)
+	if (bCEFInitialized)
 	{
 		SchemeHandlerFactories.AddSchemeHandlerFactory(MoveTemp(Scheme), MoveTemp(Domain), WebBrowserSchemeHandlerFactory);
 		return true;
@@ -1207,7 +1204,7 @@ bool FWebBrowserSingleton::RegisterSchemeHandlerFactory(FString Scheme, FString
 bool FWebBrowserSingleton::UnregisterSchemeHandlerFactory(IWebBrowserSchemeHandlerFactory* WebBrowserSchemeHandlerFactory)
 {
 #if WITH_CEF3
-	if (bAllowCEF)
+	if (bCEFInitialized)
 	{
 		SchemeHandlerFactories.RemoveSchemeHandlerFactory(WebBrowserSchemeHandlerFactory);
 		return true;
diff --git a/Engine/Source/Runtime/WebBrowser/Private/WebBrowserSingleton.h b/Engine/Source/Runtime/WebBrowser/Private/WebBrowserSingleton.h
--- a/Engine/Source/Runtime/WebBrowser/Private/WebBrowserSingleton.h
+++ b/Engine/Source/Runtime/WebBrowser/Private/WebBrowserSingleton.h
@@ -169,6 +169,11 @@ public:
 #endif
 private:
 
+	/** Critical section for thread safe modification of WindowInterfaces array. */
+	FCriticalSection WindowInterfacesCS;
+
+	TSharedRef<IWebBrowserWindowFactory> WebBrowserWindowFactory;
+
 	TSharedPtr<IWebBrowserCookieManager> DefaultCookieManager;
 
 #if WITH_CEF3
@@ -177,29 +182,22 @@ private:
 	/** Helper function that blocks until the CEF task queue has processed a posted task, flushing the queue */
 	void WaitForTaskQueueFlush();
 
+	bool bCEFInitialized;
+	bool bTaskQueueFlushed;
+
 	/** Pointer to the CEF App implementation */
 	CefRefPtr<FCEFBrowserApp>			CEFBrowserApp;
 
 	TMap<FString, CefRefPtr<CefRequestContext>> RequestContexts;
 	TMap<FString, CefRefPtr<FCEFResourceContextHandler>> RequestResourceHandlers;
 	FCefSchemeHandlerFactories SchemeHandlerFactories;
-	bool bAllowCEF;
-	bool bTaskFinished;
-#endif
 
 	/** List of currently existing browser windows */
-#if WITH_CEF3
 	TArray<TWeakPtr<FCEFWebBrowserWindow>> WindowInterfaces;
 #elif PLATFORM_IOS || PLATFORM_MAC || PLATFORM_SPECIFIC_WEB_BROWSER || (PLATFORM_ANDROID && USE_ANDROID_JNI)
+	/** List of currently existing browser windows */
 	TArray<TWeakPtr<IWebBrowserWindow>>	WindowInterfaces;
-#endif
-
-	/** Critical section for thread safe modification of WindowInterfaces array. */
-	FCriticalSection WindowInterfacesCS;
-
-	TSharedRef<IWebBrowserWindowFactory> WebBrowserWindowFactory;
-
-#if !WITH_CEF3 && (PLATFORM_IOS || PLATFORM_MAC || (PLATFORM_ANDROID && USE_ANDROID_JNI))
+	/** Application part of user agent string passed to HTTP requests*/
 	FString UserAgentApplication;
 #endif
 
@@ -209,11 +207,6 @@ private:
 
 	bool bAppIsFocused;
 
-#if WITH_CEF3
-	/** Did CEF successfully initialize itself */
-	bool bCEFInitialized;
-#endif
-
 #if WITH_ENGINE
 	/** Reference to UWebBrowser's default material*/
 	TSoftObjectPtr<UMaterialInterface> DefaultMaterial;
