WebBrowser: Allow overriding the Product/Application part of the HTTP UserAgent string per-window

With webviews being used by multiple editor plugins, some need to advertise themselves with their own product string instead of the generic project-derived one.
This is supported out of the box by the Apple and Android backends. CEF support will be added separately as it requires a new API on the CEF side.

#rnx
#Synced-CL 52548521

diff --git a/Engine/Source/Runtime/WebBrowser/Private/CEF/CEFWebBrowserWindow.cpp b/Engine/Source/Runtime/WebBrowser/Private/CEF/CEFWebBrowserWindow.cpp
--- a/Engine/Source/Runtime/WebBrowser/Private/CEF/CEFWebBrowserWindow.cpp
+++ b/Engine/Source/Runtime/WebBrowser/Private/CEF/CEFWebBrowserWindow.cpp
@@ -386,11 +386,12 @@ private:
 
 
 
-FCEFWebBrowserWindow::FCEFWebBrowserWindow(CefRefPtr<CefBrowser> InBrowser, CefRefPtr<FCEFBrowserHandler> InHandler, FString InUrl, TOptional<FString> InContentsToLoad, bool bInShowErrorMessage, bool bInThumbMouseButtonNavigation, bool bInUseTransparency, bool bInJSBindingToLoweringEnabled, bool bInUsingAcceleratedPaint)
+FCEFWebBrowserWindow::FCEFWebBrowserWindow(CefRefPtr<CefBrowser> InBrowser, CefRefPtr<FCEFBrowserHandler> InHandler, FString InUrl, TOptional<FString> InContentsToLoad, FString InUserAgentProduct, bool bInShowErrorMessage, bool bInThumbMouseButtonNavigation, bool bInUseTransparency, bool bInJSBindingToLoweringEnabled, bool bInUsingAcceleratedPaint)
 	: DocumentState(EWebBrowserDocumentState::NoDocument)
 	, InternalCefBrowser(InBrowser)
 	, WebBrowserHandler(InHandler)
 	, CurrentUrl(InUrl)
+	, UserAgentProduct(InUserAgentProduct)
 	, ViewportSize(FIntPoint::ZeroValue)
 	, ViewportDPIScaleFactor(1.0f)
 	, ViewportScaleFactor(1.0f)
diff --git a/Engine/Source/Runtime/WebBrowser/Private/CEF/CEFWebBrowserWindow.h b/Engine/Source/Runtime/WebBrowser/Private/CEF/CEFWebBrowserWindow.h
--- a/Engine/Source/Runtime/WebBrowser/Private/CEF/CEFWebBrowserWindow.h
+++ b/Engine/Source/Runtime/WebBrowser/Private/CEF/CEFWebBrowserWindow.h
@@ -91,12 +91,13 @@ private:
 	 * @param Handler Pointer to the CEF handler for this window.
 	 * @param Url The Initial URL that will be loaded.
 	 * @param ContentsToLoad Optional string to load as a web page.
+	 * @param UserAgentProduct the product part of the user agent string for HTTP requests sent by this window.
 	 * @param bShowErrorMessage Whether to show an error message in case of loading errors.
 	 * @param bThumbMouseButtonNavigation Whether to allow forward and back navigation via the mouse thumb buttons.
 	 * @param bUseTransparency Whether to enable transparency.
 	 * @param bJSBindingToLoweringEnabled Whether we ToLower all JavaScript member names.
 	 */
-	FCEFWebBrowserWindow(CefRefPtr<CefBrowser> Browser, CefRefPtr<FCEFBrowserHandler> Handler, FString Url, TOptional<FString> ContentsToLoad, bool bShowErrorMessage, bool bThumbMouseButtonNavigation, bool bUseTransparency, bool bJSBindingToLoweringEnabled, bool bUsingAcceleratedPaint);
+	FCEFWebBrowserWindow(CefRefPtr<CefBrowser> Browser, CefRefPtr<FCEFBrowserHandler> Handler, FString Url, TOptional<FString> ContentsToLoad, FString UserAgentProduct, bool bShowErrorMessage, bool bThumbMouseButtonNavigation, bool bUseTransparency, bool bJSBindingToLoweringEnabled, bool bUsingAcceleratedPaint);
 
 	/**
 	 * Create the SWidget for this WebBrowserWindow
@@ -111,6 +112,7 @@ public:
 	bool IsThumbMouseButtonNavigationEnabled() const { return bThumbMouseButtonNavigation; }
 	bool UseTransparency() const { return bUseTransparency; }
 	bool UsingAcceleratedPaint() const { return bUsingAcceleratedPaint; }
+	FString GetUserAgentProduct() const { return UserAgentProduct; }
 	
 public:
 
@@ -676,6 +678,9 @@ private:
 	/** Current tool tip. */
 	FString ToolTipText;
 
+	/** The product part for the user agent string of this window. */
+	FString UserAgentProduct;
+
 	/** Current size of this window. */
 	FIntPoint ViewportSize;
 
diff --git a/Engine/Source/Runtime/WebBrowser/Private/SWebBrowser.cpp b/Engine/Source/Runtime/WebBrowser/Private/SWebBrowser.cpp
--- a/Engine/Source/Runtime/WebBrowser/Private/SWebBrowser.cpp
+++ b/Engine/Source/Runtime/WebBrowser/Private/SWebBrowser.cpp
@@ -132,6 +132,7 @@ void SWebBrowser::Construct(const FArguments& InArgs, const TSharedPtr<IWebBrows
 				.ParentWindow(InArgs._ParentWindow)
 				.InitialURL(InArgs._InitialURL)
 				.ContentsToLoad(InArgs._ContentsToLoad)
+				.UserAgentApplication(InArgs._UserAgentApplication)
 				.ShowErrorMessage(InArgs._ShowErrorMessage)
 				.SupportsTransparency(InArgs._SupportsTransparency)
 				.SupportsThumbMouseButtonNavigation(InArgs._SupportsThumbMouseButtonNavigation)
diff --git a/Engine/Source/Runtime/WebBrowser/Private/SWebBrowserView.cpp b/Engine/Source/Runtime/WebBrowser/Private/SWebBrowserView.cpp
--- a/Engine/Source/Runtime/WebBrowser/Private/SWebBrowserView.cpp
+++ b/Engine/Source/Runtime/WebBrowser/Private/SWebBrowserView.cpp
@@ -117,6 +117,7 @@ void SWebBrowserView::Construct(const FArguments& InArgs, const TSharedPtr<IWebB
 			Settings.bInterceptLoadRequests = InArgs._InterceptLoadRequests;
 			Settings.bThumbMouseButtonNavigation = InArgs._SupportsThumbMouseButtonNavigation;
 			Settings.ContentsToLoad = InArgs._ContentsToLoad;
+			Settings.UserAgentApplication = InArgs._UserAgentApplication;
 			Settings.bShowErrorMessage = InArgs._ShowErrorMessage;
 			Settings.BackgroundColor = InArgs._BackgroundColor;
 			Settings.BrowserFrameRate = InArgs._BrowserFrameRate;
diff --git a/Engine/Source/Runtime/WebBrowser/Private/WebBrowserSingleton.cpp b/Engine/Source/Runtime/WebBrowser/Private/WebBrowserSingleton.cpp
--- a/Engine/Source/Runtime/WebBrowser/Private/WebBrowserSingleton.cpp
+++ b/Engine/Source/Runtime/WebBrowser/Private/WebBrowserSingleton.cpp
@@ -152,9 +152,10 @@ public:
 		bool bThumbMouseButtonNavigation,
 		bool bInterceptLoadRequests = true,
 		TOptional<FString> ContentsToLoad = TOptional<FString>(),
+		TOptional<FString> UserAgentApplication = TOptional<FString>(),
 		bool ShowErrorMessage = true,
 		FColor BackgroundColor = FColor(255, 255, 255, 255),
-	    bool bMobileJSReturnInDict = true) override
+		bool bMobileJSReturnInDict = true) override
 	{
 		FCreateBrowserWindowSettings Settings;
 		Settings.OSWindowHandle = OSWindowHandle;
@@ -162,6 +163,7 @@ public:
 		Settings.bUseTransparency = bUseTransparency;
 		Settings.bThumbMouseButtonNavigation = bThumbMouseButtonNavigation;
 		Settings.ContentsToLoad = MoveTemp(ContentsToLoad);
+		Settings.UserAgentApplication = MoveTemp(UserAgentApplication);
 		Settings.bShowErrorMessage = ShowErrorMessage;
 		Settings.BackgroundColor = BackgroundColor;
 		Settings.bInterceptLoadRequests = bInterceptLoadRequests;
@@ -193,15 +195,23 @@ public:
 		bool bThumbMouseButtonNavigation,
 		bool bInterceptLoadRequests = true,
 		TOptional<FString> ContentsToLoad = TOptional<FString>(),
+		TOptional<FString> UserAgentApplication = TOptional<FString>(),
 		bool ShowErrorMessage = true,
 		FColor BackgroundColor = FColor(255, 255, 255, 255),
-	    bool bMobileJSReturnInDict = true) override
+		bool bMobileJSReturnInDict = true) override
 	{
 		return nullptr;
 	}
 };
 
 #if WITH_CEF3
+
+static FString MakeUserAgentProductString(const FString& UserAgentApplication)
+{
+	// Append Chrome engine version to the product part of the user agent string
+	return FString::Printf(TEXT("%s Chrome/%d.%d.%d.%d"), *UserAgentApplication, CHROME_VERSION_MAJOR, CHROME_VERSION_MINOR, CHROME_VERSION_BUILD, CHROME_VERSION_PATCH);
+}
+
 #if PLATFORM_MAC || PLATFORM_LINUX
 class FPosixSignalPreserver
 {
@@ -286,10 +296,8 @@ FWebBrowserSingleton::FWebBrowserSingleton(const FWebBrowserInitSettings& WebBro
 	, bTaskQueueFlushed(false)
 #else
 	: WebBrowserWindowFactory(MakeShareable(new FNoWebBrowserWindowFactory()))
-#if PLATFORM_IOS || PLATFORM_MAC || (PLATFORM_ANDROID && USE_ANDROID_JNI)
-	, UserAgentApplication(WebBrowserInitSettings.ProductVersion)
-#endif
 #endif
+	, UserAgentApplication(WebBrowserInitSettings.ProductVersion)
 	, bDevToolsShortcutEnabled(UE_BUILD_DEBUG)
 	, bJSBindingsToLoweringEnabled(true)
 	, bAppIsFocused(false)
@@ -357,9 +365,8 @@ FWebBrowserSingleton::FWebBrowserSingleton(const FWebBrowserInitSettings& WebBro
 		FString LocaleCode = GetCurrentLocaleCode();
 		CefString(&Settings.locale) = TCHAR_TO_WCHAR(*LocaleCode);
 
-		// Append engine version to the user agent string.
-		FString UserAgentApplication = FString::Printf(TEXT("%s Chrome/%d.%d.%d.%d"), *WebBrowserInitSettings.ProductVersion, CHROME_VERSION_MAJOR, CHROME_VERSION_MINOR, CHROME_VERSION_BUILD, CHROME_VERSION_PATCH);
-		CefString(&Settings.user_agent_product) = TCHAR_TO_WCHAR(*UserAgentApplication);
+		FString UserAgentProduct = MakeUserAgentProductString(UserAgentApplication);
+		CefString(&Settings.user_agent_product) = TCHAR_TO_WCHAR(*UserAgentProduct);
 
 #if CEF3_DEFAULT_CACHE
 		// First clean up existing cache folders which contain outdated or problematic data
@@ -620,8 +627,9 @@ TSharedPtr<IWebBrowserWindow> FWebBrowserSingleton::CreateBrowserWindow(
 		bool bThumbMouseButtonNavigation = BrowserWindowParent->IsThumbMouseButtonNavigationEnabled();
 		bool bUseTransparency = BrowserWindowParent->UseTransparency();
 		bool bUsingAcceleratedPaint = BrowserWindowParent->UsingAcceleratedPaint();
+		FString UserAgentProduct = BrowserWindowParent->GetUserAgentProduct();
 		FString InitialURL = WCHAR_TO_TCHAR(BrowserWindowInfo->Browser->GetMainFrame()->GetURL().ToWString().c_str());
-		TSharedPtr<FCEFWebBrowserWindow> NewBrowserWindow(new FCEFWebBrowserWindow(BrowserWindowInfo->Browser, BrowserWindowInfo->Handler, InitialURL, ContentsToLoad, bShowErrorMessage, bThumbMouseButtonNavigation, bUseTransparency, bJSBindingsToLoweringEnabled, bUsingAcceleratedPaint));
+		TSharedPtr<FCEFWebBrowserWindow> NewBrowserWindow(new FCEFWebBrowserWindow(BrowserWindowInfo->Browser, BrowserWindowInfo->Handler, InitialURL, ContentsToLoad, UserAgentProduct, bShowErrorMessage, bThumbMouseButtonNavigation, bUseTransparency, bJSBindingsToLoweringEnabled, bUsingAcceleratedPaint));
 		BrowserWindowInfo->Handler->SetBrowserWindow(NewBrowserWindow);
 		{
 			FScopeLock Lock(&WindowInterfacesCS);
@@ -757,6 +765,7 @@ TSharedPtr<IWebBrowserWindow> FWebBrowserSingleton::CreateBrowserWindow(const FC
 				NewHandler,
 				WindowSettings.InitialURL,
 				WindowSettings.ContentsToLoad,
+				MakeUserAgentProductString(WindowSettings.UserAgentApplication.Get(UserAgentApplication /*DefaultValue*/)),
 				WindowSettings.bShowErrorMessage,
 				WindowSettings.bThumbMouseButtonNavigation,
 				WindowSettings.bUseTransparency,
@@ -781,8 +790,8 @@ TSharedPtr<IWebBrowserWindow> FWebBrowserSingleton::CreateBrowserWindow(const FC
 		WindowSettings.bThumbMouseButtonNavigation,
 		WindowSettings.bUseTransparency,
 		bJSBindingsToLoweringEnabled,
-		UserAgentApplication,
-	    WindowSettings.bMobileJSReturnInDict));
+		WindowSettings.UserAgentApplication.Get(UserAgentApplication /*DefaultValue*/),
+		WindowSettings.bMobileJSReturnInDict));
 
 	{
 		FScopeLock Lock(&WindowInterfacesCS);
@@ -798,8 +807,8 @@ TSharedPtr<IWebBrowserWindow> FWebBrowserSingleton::CreateBrowserWindow(const FC
 		WindowSettings.bThumbMouseButtonNavigation, 
 		WindowSettings.bUseTransparency,
 		bJSBindingsToLoweringEnabled,
-		UserAgentApplication,
-	    WindowSettings.bMobileJSReturnInDict));
+		WindowSettings.UserAgentApplication.Get(UserAgentApplication /*DefaultValue*/),
+		WindowSettings.bMobileJSReturnInDict));
 
 	{
 		FScopeLock Lock(&WindowInterfacesCS);
diff --git a/Engine/Source/Runtime/WebBrowser/Public/IWebBrowserSingleton.h b/Engine/Source/Runtime/WebBrowser/Public/IWebBrowserSingleton.h
--- a/Engine/Source/Runtime/WebBrowser/Public/IWebBrowserSingleton.h
+++ b/Engine/Source/Runtime/WebBrowser/Public/IWebBrowserSingleton.h
@@ -27,9 +27,10 @@ public:
 		bool bThumbMouseButtonNavigation,
 		bool InterceptLoadRequests = true,
 		TOptional<FString> ContentsToLoad = TOptional<FString>(),
+		TOptional<FString> UserAgentApplication = TOptional<FString>(),
 		bool ShowErrorMessage = true,
 		FColor BackgroundColor = FColor(255, 255, 255, 255),
-	    bool bMobileJSReturnInDict = true) = 0;
+		bool bMobileJSReturnInDict = true) = 0;
 };
 
 struct FBrowserContextSettings
@@ -55,34 +56,19 @@ struct FBrowserContextSettings
 
 struct FCreateBrowserWindowSettings
 {
-
-	FCreateBrowserWindowSettings()
-		: OSWindowHandle(nullptr)
-		, InitialURL()
-		, bUseTransparency(false)
-		, bInterceptLoadRequests(true)
-		, bThumbMouseButtonNavigation(false)
-		, ContentsToLoad()
-		, bShowErrorMessage(true)
-		, BackgroundColor(FColor(255, 255, 255, 255))
-		, BrowserFrameRate(24)
-		, Context()
-		, AltRetryDomains()
-		, bMobileJSReturnInDict(true)
-	{ }
-
-	void* OSWindowHandle;
+	void* OSWindowHandle = nullptr;
 	FString InitialURL;
-	bool bUseTransparency;
-	bool bInterceptLoadRequests;
-	bool bThumbMouseButtonNavigation;
-	TOptional<FString> ContentsToLoad;
-	bool bShowErrorMessage;
-	FColor BackgroundColor;
-	int BrowserFrameRate;
+	FColor BackgroundColor = FColor(255, 255, 255, 255);
+	int BrowserFrameRate = 24;
+	bool bUseTransparency = false;
+	bool bInterceptLoadRequests = true;
+	bool bThumbMouseButtonNavigation = false;
+	bool bShowErrorMessage = true;
+	bool bMobileJSReturnInDict = true;
 	TOptional<FBrowserContextSettings> Context;
+	TOptional<FString> ContentsToLoad;
+	TOptional<FString> UserAgentApplication;
 	TArray<FString> AltRetryDomains;
-	bool bMobileJSReturnInDict;
 };
 
 /**
diff --git a/Engine/Source/Runtime/WebBrowser/Public/SWebBrowser.h b/Engine/Source/Runtime/WebBrowser/Public/SWebBrowser.h
--- a/Engine/Source/Runtime/WebBrowser/Public/SWebBrowser.h
+++ b/Engine/Source/Runtime/WebBrowser/Public/SWebBrowser.h
@@ -51,6 +51,9 @@ public:
 		/** Optional string to load contents as a web page. */
 		SLATE_ARGUMENT(TOptional<FString>, ContentsToLoad)
 
+		/** Custom Application part of the user agent for HTTP requests. If not set, WebBrowserInitSettings.ProductVersion will be used. */
+		SLATE_ARGUMENT(TOptional<FString>, UserAgentApplication)
+
 		/** Whether to show standard controls like Back, Forward, Reload etc. */
 		SLATE_ARGUMENT(bool, ShowControls)
 
diff --git a/Engine/Source/Runtime/WebBrowser/Public/SWebBrowserView.h b/Engine/Source/Runtime/WebBrowser/Public/SWebBrowserView.h
--- a/Engine/Source/Runtime/WebBrowser/Public/SWebBrowserView.h
+++ b/Engine/Source/Runtime/WebBrowser/Public/SWebBrowserView.h
@@ -72,6 +72,9 @@ public:
 		/** Optional string to load contents as a web page. */
 		SLATE_ARGUMENT(TOptional<FString>, ContentsToLoad)
 
+		/** Custom Application part of the user agent for HTTP requests. If not set, WebBrowserInitSettings.ProductVersion will be used. */
+		SLATE_ARGUMENT(TOptional<FString>, UserAgentApplication)
+
 		/** Whether to show an error message in case of loading errors. */
 		SLATE_ARGUMENT(bool, ShowErrorMessage)
 
