WebBrowser: Fix CEF losing mouse events when the cursor leaves the window, even when still captured by UE

Also adjust mouse capture behavior in WebBrowserViewport to allow middle and right mouse buttons too, and to only capture (resp. release) the mouse capture on the first button being pressed (resp. last button being released) for drags with multiple buttons.

#rnx
#Synced-CL 52583945

diff --git a/Engine/Source/Runtime/WebBrowser/Private/CEF/CEFWebBrowserWindow.cpp b/Engine/Source/Runtime/WebBrowser/Private/CEF/CEFWebBrowserWindow.cpp
--- a/Engine/Source/Runtime/WebBrowser/Private/CEF/CEFWebBrowserWindow.cpp
+++ b/Engine/Source/Runtime/WebBrowser/Private/CEF/CEFWebBrowserWindow.cpp
@@ -1619,7 +1619,7 @@ FReply FCEFWebBrowserWindow::OnMouseMove(const FGeometry& MyGeometry, const FPoi
 
 		if (!bEventConsumedByDragCallback)
 		{
-			InternalCefBrowser->GetHost()->SendMouseMoveEvent(Event, false);
+			InternalCefBrowser->GetHost()->SendMouseMoveEvent(Event, false /*mouseLeave*/);
 		}
 		
 		Reply = FReply::Handled();
@@ -1631,13 +1631,15 @@ void FCEFWebBrowserWindow::OnMouseLeave(const FPointerEvent& MouseEvent)
 {
 	// Ensure we clear any tooltips if the mouse leaves the window.
 	SetToolTip(CefString());
-	// We have no geometry here to convert our mouse event to local space so we just make a dummy event and set the moueLeave param to true
-	CefMouseEvent DummyEvent;
-	if (IsValid() && !BlockInputInDirectHwndMode())
+
+	// Don't send mouseLeave if we still have mouse buttons pressed, it indicates an outgoing drag with the mouse captured.
+	// We'll receive another OnMouseLeave if we're still outside when the last button has been released.
+	if (IsValid() && !BlockInputInDirectHwndMode() && MouseEvent.GetPressedButtons().IsEmpty())
 	{
-		InternalCefBrowser->GetHost()->SendMouseMoveEvent(DummyEvent, true);
+		// We have no geometry here to convert our mouse event to local space so we just make a dummy event and set the mouseLeave param to true
+		CefMouseEvent DummyEvent;
+		InternalCefBrowser->GetHost()->SendMouseMoveEvent(DummyEvent, true /*mouseLeave*/);
 	}
-
 }
 
 void FCEFWebBrowserWindow::SetSupportsMouseWheel(bool bValue)
diff --git a/Engine/Source/Runtime/WebBrowser/Private/WebBrowserViewport.cpp b/Engine/Source/Runtime/WebBrowser/Private/WebBrowserViewport.cpp
--- a/Engine/Source/Runtime/WebBrowser/Private/WebBrowserViewport.cpp
+++ b/Engine/Source/Runtime/WebBrowser/Private/WebBrowserViewport.cpp
@@ -51,9 +51,10 @@ FCursorReply FWebBrowserViewport::OnCursorQuery( const FGeometry& MyGeometry, co
 
 FReply FWebBrowserViewport::OnMouseButtonDown(const FGeometry& MyGeometry, const FPointerEvent& MouseEvent)
 {
-	// Capture mouse on left button down so that you can drag out of the viewport
+	// Capture mouse on left/middle/right button down so that you can drag out of the viewport
 	FReply Reply = WebBrowserWindow->OnMouseButtonDown(MyGeometry, MouseEvent, bIsPopup);
-	if (MouseEvent.GetEffectingButton() == EKeys::LeftMouseButton)
+	if (MouseEvent.GetPressedButtons().Num() == 1 &&
+		(MouseEvent.GetEffectingButton() == EKeys::LeftMouseButton || MouseEvent.GetEffectingButton() == EKeys::MiddleMouseButton || MouseEvent.GetEffectingButton() == EKeys::RightMouseButton))
 	{
 		const FWidgetPath* Path = MouseEvent.GetEventPath();
 		if (Path->IsValid())
@@ -67,9 +68,9 @@ FReply FWebBrowserViewport::OnMouseButtonDown(const FGeometry& MyGeometry, const
 
 FReply FWebBrowserViewport::OnMouseButtonUp(const FGeometry& MyGeometry, const FPointerEvent& MouseEvent)
 {
-	// Release mouse capture when left button released
+	// Release mouse capture when all buttons released
 	FReply Reply = WebBrowserWindow->OnMouseButtonUp(MyGeometry, MouseEvent, bIsPopup);
-	if (MouseEvent.GetEffectingButton() == EKeys::LeftMouseButton)
+	if (MouseEvent.GetPressedButtons().IsEmpty())
 	{
 		return Reply.ReleaseMouseCapture();
 	}
