WebBrowser: forward per-browser UserAgentApplication to CEF browsers

#rnx
#Synced-CL 52747226

diff --git a/Engine/Source/Runtime/WebBrowser/Private/WebBrowserSingleton.cpp b/Engine/Source/Runtime/WebBrowser/Private/WebBrowserSingleton.cpp
--- a/Engine/Source/Runtime/WebBrowser/Private/WebBrowserSingleton.cpp
+++ b/Engine/Source/Runtime/WebBrowser/Private/WebBrowserSingleton.cpp
@@ -707,6 +707,13 @@ TSharedPtr<IWebBrowserWindow> FWebBrowserSingleton::CreateBrowserWindow(const FC
 			BrowserSettings.windowless_frame_rate = BrowserFrameRate;
 		}
 
+		FString UserAgentProduct = MakeUserAgentProductString(WindowSettings.UserAgentApplication.Get(UserAgentApplication /*DefaultValue*/));
+		// The user_agent_product setting is only available for Windows and Linux at the moment,
+		// as it requires updated EpicWebHelper binaries and that's not possible anymore on Mac until we update CEF
+#if PLATFORM_WINDOWS || PLATFORM_LINUX
+		CefString(&BrowserSettings.user_agent_product) = TCHAR_TO_WCHAR(*UserAgentProduct);
+#endif
+
 		TArray<FString> AuthorizationHeaderAllowListURLS;
 		GConfig->GetArray(TEXT("Browser"), TEXT("AuthorizationHeaderAllowListURLS"), AuthorizationHeaderAllowListURLS, GEngineIni);
 
@@ -760,7 +767,7 @@ TSharedPtr<IWebBrowserWindow> FWebBrowserSingleton::CreateBrowserWindow(const FC
 				NewHandler,
 				WindowSettings.InitialURL,
 				WindowSettings.ContentsToLoad,
-				MakeUserAgentProductString(WindowSettings.UserAgentApplication.Get(UserAgentApplication /*DefaultValue*/)),
+				UserAgentProduct,
 				WindowSettings.bShowErrorMessage,
 				WindowSettings.bThumbMouseButtonNavigation,
 				WindowSettings.bUseTransparency,
