WebBrowser: Add a MakeUserAgentApplication helper to IWebBrowserModule

#rnx
#Synced-CL 52782919

diff --git a/Engine/Source/Runtime/WebBrowser/Private/WebBrowserModule.cpp b/Engine/Source/Runtime/WebBrowser/Private/WebBrowserModule.cpp
--- a/Engine/Source/Runtime/WebBrowser/Private/WebBrowserModule.cpp
+++ b/Engine/Source/Runtime/WebBrowser/Private/WebBrowserModule.cpp
@@ -18,8 +18,13 @@ DEFINE_LOG_CATEGORY(LogWebBrowser);
 
 static FWebBrowserSingleton* WebBrowserSingleton = nullptr;
 
+FString IWebBrowserModule::MakeUserAgentApplication(const FString& ApplicationName, const FString& ApplicationVersion)
+{
+	return FString::Printf(TEXT("%s/%s UnrealEngine/%s"), *ApplicationName, ApplicationVersion.IsEmpty() ? TEXT("1.0") : *ApplicationVersion, *FEngineVersion::Current().ToString());
+}
+
 FWebBrowserInitSettings::FWebBrowserInitSettings()
-	: ProductVersion(FString::Printf(TEXT("%s/%s UnrealEngine/%s"), FApp::GetProjectName(), FApp::GetBuildVersion(), *FEngineVersion::Current().ToString()))
+	: ProductVersion(IWebBrowserModule::MakeUserAgentApplication(FApp::GetProjectName(), FApp::GetBuildVersion()))
 {
 }
 
diff --git a/Engine/Source/Runtime/WebBrowser/Public/WebBrowserModule.h b/Engine/Source/Runtime/WebBrowser/Public/WebBrowserModule.h
--- a/Engine/Source/Runtime/WebBrowser/Public/WebBrowserModule.h
+++ b/Engine/Source/Runtime/WebBrowser/Public/WebBrowserModule.h
@@ -49,6 +49,16 @@ public:
 		return FModuleManager::Get().IsModuleLoaded("WebBrowser");
 	}
 
+	/**
+	 * Returns a string suitable for use as the application (or product) part of a browser's user agent string.
+	 * Includes the formatted application/product name and version, followed by the engine name and version.
+	 *
+	 * @param ApplicationName the application/product name
+	 * @param ApplicationVersion the application/product version, will default to 1.0 if empty
+	 * @return The user agent application string
+	 */
+	WEBBROWSER_API static FString MakeUserAgentApplication(const FString& ApplicationName, const FString& ApplicationVersion = FString());
+
 	/**
 	 * Customize initialization settings. You must call this before the first GetSingleton call, in order to override init settings.
 	 * 
