WebBrowser: remove obsolete log warning about crashes when closing a webview

#rnx
#Synced-CL 52782919

diff --git a/Engine/Source/Runtime/WebBrowser/Private/CEF/CEFWebBrowserWindow.cpp b/Engine/Source/Runtime/WebBrowser/Private/CEF/CEFWebBrowserWindow.cpp
--- a/Engine/Source/Runtime/WebBrowser/Private/CEF/CEFWebBrowserWindow.cpp
+++ b/Engine/Source/Runtime/WebBrowser/Private/CEF/CEFWebBrowserWindow.cpp
@@ -577,7 +577,6 @@ FCEFWebBrowserWindow::~FCEFWebBrowserWindow()
 	WebBrowserHandler->OnConsoleMessage().Unbind();
 	if (IsValid())
 	{
-		UE_LOGF(LogWebBrowser, Log, "Closing browser during destruction, this may cause a later crash.");
 		CloseBrowser(true, false);
 	}
 
